#UE4 First iteration on the merge module.. Changes to diff panel reviewed by [REDACTED_PERSON], driver disabled by [AssetMerge] EnableAssetMerge

diff --git a/Engine/Source/Developer/AssetTools/Public/IAssetTypeActions.h b/Engine/Source/Developer/AssetTools/Public/IAssetTypeActions.h
--- a/Engine/Source/Developer/AssetTools/Public/IAssetTypeActions.h
+++ b/Engine/Source/Developer/AssetTools/Public/IAssetTypeActions.h
@@ -36,6 +36,14 @@ struct FRevisionInfo
 	int32		Revision;
 	int32		Changelist;
 	FDateTime	Date;	
+
+	static FRevisionInfo InvalidRevision()
+	{
+		static const FRevisionInfo InvalidRevision{
+			-1, -1, FDateTime()
+		};
+		return InvalidRevision;
+	}
 };
 
 /** AssetTypeActions provide actions and other information about asset types */
diff --git a/Engine/Source/Developer/Merge/Merge.Build.cs b/Engine/Source/Developer/Merge/Merge.Build.cs
new file mode 100644
--- /dev/null
+++ b/Engine/Source/Developer/Merge/Merge.Build.cs
@@ -0,0 +1,52 @@
+// Copyright 1998-2014 Epic Games, Inc. All Rights Reserved.
+
+namespace UnrealBuildTool.Rules
+{
+	public class Merge : ModuleRules
+	{
+        public Merge(TargetInfo Target)
+		{
+			PublicIncludePaths.AddRange(
+				new string[] {
+				}
+				);
+
+			PrivateIncludePaths.AddRange(
+				new string[] {
+					"Developer/Merge/Private",
+				    "Kismet",
+				}
+				);
+
+			PublicDependencyModuleNames.AddRange(
+				new string[]
+				{
+					"Core",
+				}
+				);
+
+			PrivateDependencyModuleNames.AddRange(
+				new string[]
+				{
+				    "AssetTools",
+				    "CoreUObject",
+				    "EditorStyle",
+				    "Engine", // needed so that we can clone blueprints...
+				    "GraphEditor",
+				    "InputCore",
+				    "Kismet",
+				    "Slate",
+				    "SlateCore",
+				    "SourceControl",
+				    "UnrealEd",
+				}
+				);
+
+			DynamicallyLoadedModuleNames.AddRange(
+				new string[]
+				{
+				}
+                );
+		}
+	}
+}
\ No newline at end of file
diff --git a/Engine/Source/Developer/Merge/Private/Merge.cpp b/Engine/Source/Developer/Merge/Private/Merge.cpp
new file mode 100644
--- /dev/null
+++ b/Engine/Source/Developer/Merge/Private/Merge.cpp
@@ -0,0 +1,245 @@
+// Copyright 1998-2014 Epic Games, Inc. All Rights Reserved.
+
+#include "MergePrivatePCH.h"
+#include "SBlueprintMerge.h"
+#include "Toolkits/ToolkitManager.h"
+#include "Toolkits/IToolkit.h"
+#include "ISourceControlModule.h"
+
+#define LOCTEXT_NAMESPACE "Merge"
+
+static FSourceControlStatePtr GetSourceControlState(const FString& PackageName)
+{
+	ISourceControlProvider& SourceControlProvider = ISourceControlModule::Get().GetProvider();
+	TSharedRef<FUpdateStatus, ESPMode::ThreadSafe> UpdateStatusOperation = ISourceControlOperation::Create<FUpdateStatus>();
+	UpdateStatusOperation->SetUpdateHistory(true);
+	SourceControlProvider.Execute(UpdateStatusOperation, SourceControlHelpers::PackageFilename(PackageName));
+
+	FSourceControlStatePtr State = SourceControlProvider.GetState(SourceControlHelpers::PackageFilename(PackageName), EStateCacheUsage::Use);
+	if (!State.IsValid() || !State->IsSourceControlled() || !FPackageName::DoesPackageExist(PackageName))
+	{
+		return FSourceControlStatePtr();
+	}
+	else
+	{
+		return State;
+	}
+}
+
+static UObject* LoadRevision(const FString& AssetName, const ISourceControlRevision& DesiredRevision)
+{
+	// Get the head revision of this package from source control
+	FString TempFileName;
+	if (DesiredRevision.Get(TempFileName))
+	{
+		// Forcibly disable compile on load in case we are loading old blueprints that might try to update/compile
+		TGuardValue<bool> DisableCompileOnLoad(GForceDisableBlueprintCompileOnLoad, true);
+
+		// Try and load that package
+		UPackage* TempPackage = LoadPackage(NULL, *TempFileName, LOAD_None);
+		if (TempPackage != NULL)
+		{
+			// Grab the old asset from that old package
+			UObject* OldObject = FindObject<UObject>(TempPackage, *AssetName);
+			if (OldObject)
+			{
+				return OldObject;
+			}
+			else
+			{
+				// @todo DO: toast? i love toast..
+			}
+		}
+		else
+		{
+			// @todo DO: toast? i love toast..
+		}
+	}
+	else
+	{
+		// @todo DO: toast? i love toast..
+	}
+	return NULL;
+}
+
+static FRevisionInfo GetRevisionInfo(ISourceControlRevision const& FromRevision)
+{
+	return{ FromRevision.GetRevisionNumber(), FromRevision.GetCheckInIdentifier(), FromRevision.GetDate() };
+}
+
+static UObject* LoadHeadRev(const FString& PackageName, const FString& AssetName, FRevisionInfo& OutRevInfo)
+{
+	FSourceControlStatePtr SourceControlState = GetSourceControlState(PackageName);
+	if (SourceControlState.IsValid())
+	{
+		// HistoryItem(0) is apparently the head revision:
+		TSharedPtr<ISourceControlRevision, ESPMode::ThreadSafe> Revision = SourceControlState->GetHistoryItem(0);
+		check(Revision.IsValid());
+		OutRevInfo = GetRevisionInfo(*Revision);
+		return LoadRevision(AssetName, *Revision);
+	}
+	else
+	{
+		// @todo DO: toast? i love toast..
+	}
+
+	return NULL;
+}
+
+static UObject* LoadBaseRev(const FString& PackageName, const FString& AssetName, FRevisionInfo& OutRevInfo)
+{
+	FSourceControlStatePtr SourceControlState = GetSourceControlState(PackageName);
+	if (SourceControlState.IsValid())
+	{
+		TSharedPtr<ISourceControlRevision, ESPMode::ThreadSafe> Revision = SourceControlState->GetBaseRevForMerge();
+		check(Revision.IsValid());
+		OutRevInfo = GetRevisionInfo(*Revision);
+		return LoadRevision(AssetName, *Revision);
+	}
+	else
+	{
+		// @todo DO: toast? i love toast..
+	}
+	return NULL;
+}
+
+class FMerge : public IMerge
+{
+	/** IModuleInterface implementation */
+	virtual void StartupModule() override;
+	virtual void ShutdownModule() override;
+	virtual void Merge(const UBlueprint& Object) override;
+};
+IMPLEMENT_MODULE( FMerge, Merge )
+
+void FMerge::StartupModule()
+{
+	// This code will execute after your module is loaded into memory (but after global variables are initialized, of course.)
+}
+
+void FMerge::ShutdownModule()
+{
+	// This function may be called during shutdown to clean up your module.  For modules that support dynamic reloading,
+	// we call this function before unloading the module.
+}
+
+void FMerge::Merge( const UBlueprint& Object )
+{
+	// Verify that the object is not open in some other editor:
+	TSharedPtr< IToolkit > FoundAssetEditor = FToolkitManager::Get().FindEditorForAsset(&Object);
+	if (FoundAssetEditor.IsValid())
+	{
+		FText ToolkitName = FoundAssetEditor->GetBaseToolkitName();
+		FText ErrorMessage = FText::Format(LOCTEXT("CloseOpenEditors", "{0} must be closed before merging {1}. Would you like to close {0} and continue?"), ToolkitName, FText::FromString(Object.GetName()));
+		EAppReturnType::Type Result = FMessageDialog::Open(EAppMsgType::OkCancel, ErrorMessage);
+		if (Result == EAppReturnType::Cancel)
+		{
+			return;
+		}
+		else
+		{
+			FToolkitManager::Get().CloseToolkit(FoundAssetEditor.ToSharedRef());
+			FoundAssetEditor = FToolkitManager::Get().FindEditorForAsset(&Object);
+			if (FoundAssetEditor.IsValid())
+			{
+				ErrorMessage = FText::Format(LOCTEXT("MergeAborted", "Aborted Merge of {0} because {1} failed to close"), FText::FromString(Object.GetName()), ToolkitName);
+				FNotificationInfo Info(ErrorMessage);
+				Info.ExpireDuration = 5.0f;
+				FSlateNotificationManager::Get().AddNotification(Info);
+				return;
+			}
+		}
+	}
+
+	// verify that the object is not dirty:
+	if (Object.GetOutermost()->IsDirty())
+	{
+		FText RequestSave = FText::Format(LOCTEXT("RequestSaveForMerge", "{0} must be saved before it can be merged. Would you like to save and continue?"), FText::FromString(Object.GetName()));
+		EAppReturnType::Type Result = FMessageDialog::Open(EAppMsgType::OkCancel, RequestSave);
+		if (Result == EAppReturnType::Cancel)
+		{
+			return;
+		}
+		else
+		{
+			TArray<UPackage*> PackagesToSave;
+			PackagesToSave.Add(Object.GetOutermost());
+			const auto Result = FEditorFileUtils::PromptForCheckoutAndSave(PackagesToSave, /*bCheckDirty=*/ false, /*bPromptToSave=*/ false);
+			if (Result != FEditorFileUtils::PR_Success)
+			{
+				return;
+			}
+		}
+	}
+
+	// merge the local asset with the depot, SCC provides us with the last common revision as
+	// a basis for the merge:
+
+	// @todo DO: this will probably need to be async.. pulling down some old versions of assets:
+	const FString& PackageName = Object.GetOutermost()->GetName();
+	const FString& AssetName = Object.GetName();
+
+	FRevisionInfo CurrentRevInfo = FRevisionInfo::InvalidRevision();
+	UBlueprint* RemoteBlueprint = Cast< UBlueprint >( LoadHeadRev(PackageName, AssetName, CurrentRevInfo) );
+	FRevisionInfo BaseRevInfo = FRevisionInfo::InvalidRevision();
+	UBlueprint* BaseBlueprint = Cast< UBlueprint >( LoadBaseRev(PackageName, AssetName, BaseRevInfo) );
+
+	if (RemoteBlueprint && BaseBlueprint)
+	{
+		FMergeDisplayArgs DisplayArgs
+		{
+			CurrentRevInfo,
+			BaseRevInfo
+		};
+
+		FText WindowTitle = FText::Format(LOCTEXT("BlueprintMerge", "{0} - Blueprint Merge"), FText::FromString(AssetName));
+
+		const TSharedPtr<SWindow> Window = SNew(SWindow)
+			.Title(WindowTitle)
+			.ClientSize(FVector2D(1000, 800));
+
+		SBlueprintDiff::FArguments BaseArgs;
+		BaseArgs.BlueprintOld(BaseBlueprint)
+			.OldRevision(CurrentRevInfo)
+			.BlueprintNew(RemoteBlueprint)
+			.NewRevision(BaseRevInfo);
+
+		Window->SetContent(SNew(SBlueprintMerge)
+			.BlueprintLocal(const_cast<UBlueprint*>(&Object))
+			.OwningWindow(Window)
+			.BaseArgs(BaseArgs));
+
+		// Make this window a child of the modal window if we've been spawned while one is active.
+		TSharedPtr<SWindow> ActiveModal = FSlateApplication::Get().GetActiveModalWindow();
+		if (ActiveModal.IsValid())
+		{
+			FSlateApplication::Get().AddWindowAsNativeChild(Window.ToSharedRef(), ActiveModal.ToSharedRef());
+		}
+		else
+		{
+			FSlateApplication::Get().AddWindow(Window.ToSharedRef());
+		}
+	}
+	else
+	{
+		FText MissingFiles;
+		if (!RemoteBlueprint && !BaseBlueprint)
+		{
+			MissingFiles = LOCTEXT("MergeBothRevisionsFailed", "common base, nor conflicting revision");
+		}
+		else if (!RemoteBlueprint)
+		{
+			MissingFiles = LOCTEXT("MergeConflictingRevisionsFailed", "conflicting revision");
+		}
+		else
+		{
+			MissingFiles = LOCTEXT("MergeBaseRevisionsFailed", "common base");
+		}
+
+		FText ErrorMessage = FText::Format(LOCTEXT("MergeRevisionLoadFailed", "Aborted Merge of {0} because we could not load {1}"), FText::FromString(Object.GetName()), MissingFiles);
+		FNotificationInfo Info(ErrorMessage);
+		Info.ExpireDuration = 3.0f;
+		FSlateNotificationManager::Get().AddNotification(Info);
+	}
+}
+
diff --git a/Engine/Source/Developer/Merge/Private/MergePrivatePCH.h b/Engine/Source/Developer/Merge/Private/MergePrivatePCH.h
new file mode 100644
--- /dev/null
+++ b/Engine/Source/Developer/Merge/Private/MergePrivatePCH.h
@@ -0,0 +1,7 @@
+// Copyright 1998-2014 Epic Games, Inc. All Rights Reserved.
+
+#include "Merge.h"
+#include "SlateCore.h"
+#include "UnrealEd.h"
+
+
diff --git a/Engine/Source/Developer/Merge/Private/SBlueprintMerge.cpp b/Engine/Source/Developer/Merge/Private/SBlueprintMerge.cpp
new file mode 100644
--- /dev/null
+++ b/Engine/Source/Developer/Merge/Private/SBlueprintMerge.cpp
@@ -0,0 +1,669 @@
+// Copyright 1998-2014 Epic Games, Inc. All Rights Reserved.
+
+#include "MergePrivatePCH.h"
+#include "SBlueprintMerge.h"
+#include "DiffResults.h"
+#include "GraphDiffControl.h"
+#include "ISourceControlModule.h"
+
+#define LOCTEXT_NAMESPACE "SBlueprintDif"
+
+// todo doc: where doe this utility function belong!
+static UPackage* CreateTempPackage(FString Name)
+{
+	static uint32 TempUid = 0;
+	FString TempPackageName = FString::Printf(TEXT("/Temp/BpMerge-%u-%s"), TempUid++, *Name);
+	return CreatePackage(NULL, *TempPackageName);
+}
+
+static UBlueprint* DuplicateBlueprint( const UBlueprint* const BlueprintToClone)
+{
+	UPackage* TempPackage = CreateTempPackage(BlueprintToClone->GetName());
+	TArray<UObject*> Subobjects;
+	GetObjectsWithOuter(TempPackage, Subobjects, true);
+	check(Subobjects.Num() == 0);
+	EObjectFlags FlagMask = RF_AllFlags & ~RF_Transient & ~RF_PendingKill;
+	return Cast<UBlueprint>(StaticDuplicateObject(BlueprintToClone, TempPackage, *BlueprintToClone->GetName(), FlagMask));
+}
+
+template< class Predicate >
+UObject* FindOuter( UObject& Obj, Predicate P )
+{
+	UObject* Iter = &Obj;
+	while( Iter )
+	{
+		const UObject* ImmutableIter = Iter;
+		if( P(*ImmutableIter) )
+		{
+			break;
+		}
+		Iter = Iter->GetOuter();
+	}
+	return Iter;
+}
+
+static bool PackagesAppearEquivalent( const UPackage* A, const UPackage* B )
+{
+	TArray<UObject*> SubobjectsA;
+	GetObjectsWithOuter(A, SubobjectsA, true);
+	TArray<UObject*> SubobjectsB;
+	GetObjectsWithOuter(B, SubobjectsB, true);
+
+	auto IsTransientOrPendingKill = [](UObject const& Obj){ return Obj.HasAnyFlags(RF_Transient | RF_PendingKill); };
+	TSet<FString> SubobjectANames;
+	for( const auto& I : SubobjectsA )
+	{
+		if (!FindOuter(*I, IsTransientOrPendingKill ) )
+		{
+			SubobjectANames.Add(I->GetName());
+		}
+	}
+	TSet<FString> SubobjectBNames;
+	for (const auto& I : SubobjectsB)
+	{
+		if (!FindOuter(*I, IsTransientOrPendingKill))
+		{
+			SubobjectBNames.Add(I->GetName());
+		}
+	}
+
+	TArray<UObject*> DifferentObjects;
+	TSet<FString> AB = SubobjectANames.Difference(SubobjectBNames);
+	TSet<FString> BA = SubobjectBNames.Difference(SubobjectANames);
+	TSet<FString> Difference = AB.Union( BA );
+	for(const auto& I : Difference )
+	{
+		auto FindByName = [&I]( const UObject* Candidate ) { return Candidate->GetName() == I; };
+		UObject** FoundObject = SubobjectsA.FindByPredicate( FindByName );
+		if( !FoundObject )
+		{
+			FoundObject = SubobjectsB.FindByPredicate( FindByName );
+		}
+
+		if( FoundObject )
+		{
+			DifferentObjects.Add( *FoundObject );
+		}
+		else
+		{
+			verify(false);
+		}
+	}
+
+	return Difference.Num() == 0;
+}
+
+/** 
+ * Commands used to navigate the lists of differences between the local revision and the common base
+ * or the remote revision and the common base.
+ */
+class FMergeDifferencesListCommands : public TCommands<FMergeDifferencesListCommands>
+{
+public:
+	FMergeDifferencesListCommands()
+		: TCommands<FMergeDifferencesListCommands>("MergeList", NSLOCTEXT("Merge", "Blueprint", "Blueprint Merge"), NAME_None, FEditorStyle::GetStyleSetName())
+	{
+	}
+
+	/** Go to previous or next difference, focusing on either the remote or local diff */
+	TSharedPtr<FUICommandInfo> LocalPrevious;
+	TSharedPtr<FUICommandInfo> LocalNext;
+	TSharedPtr<FUICommandInfo> RemotePrevious;
+	TSharedPtr<FUICommandInfo> RemoteNext;
+
+	virtual void RegisterCommands() override
+	{
+		UI_COMMAND(LocalPrevious, "LocalPrevious", "Go to previous local difference", EUserInterfaceActionType::Button, FInputGesture(EKeys::F7, EModifierKey::Control));
+		UI_COMMAND(LocalNext, "LocalNext", "Go to next local difference", EUserInterfaceActionType::Button, FInputGesture(EKeys::F7));
+		UI_COMMAND(RemotePrevious, "RemotePrevious", "Go to previous remote difference", EUserInterfaceActionType::Button, FInputGesture(EKeys::F8, EModifierKey::Control));
+		UI_COMMAND(RemoteNext, "RemoteNext", "Go to next remote difference", EUserInterfaceActionType::Button, FInputGesture(EKeys::F8));
+	}
+};
+
+void SBlueprintMerge::Construct(const FArguments& InArgs)
+{
+	FMergeDifferencesListCommands::Register();
+
+	PanelLocal.Blueprint = InArgs._BlueprintLocal;
+	BlueprintResult = DuplicateBlueprint(InArgs._BlueprintLocal);
+	(void)PackagesAppearEquivalent( BlueprintResult->GetOutermost(), PanelLocal.Blueprint->GetOutermost() );
+	//check( PackagesAppearEquivalent( BlueprintResult->GetOutermost(), PanelLocal.Blueprint->GetOutermost() ) );
+	OwningWindow = InArgs._OwningWindow;
+
+	return SBlueprintDiff::Construct( InArgs._BaseArgs );
+}
+
+TSharedRef<ITableRow> SBlueprintMerge::OnGenerateRow( TSharedPtr<struct FDiffSingleResult> ParamItem, const TSharedRef<STableViewBase>& OwnerTable)
+{
+	FDiffSingleResult const& Result = *ParamItem;
+	FString ToolTip = Result.ToolTip;
+	FLinearColor Color = Result.DisplayColor;
+	FString Text = Result.DisplayString;
+	if (Text.Len() == 0)
+	{
+		Text = LOCTEXT("DIF_UnknownDiff", "Unknown Diff").ToString();
+		ToolTip = LOCTEXT("DIF_Confused", "There is an unspecified difference").ToString();
+	}
+	return SNew(STableRow< TSharedPtr<struct FDiffSingleResult> >, OwnerTable)
+		.Content()
+		[
+			SNew(STextBlock)
+			.ToolTipText(ToolTip)
+			.ColorAndOpacity(Color)
+			.Text(Text)
+		];
+}
+
+void SBlueprintMerge::OnDiffListSelectionChanged(TSharedPtr<struct FDiffSingleResult> Item, ESelectInfo::Type SelectionType)
+{
+	if( !Item.IsValid() )
+	{
+		return;
+	}
+
+	//focus the graph onto the diff that was clicked on
+	FDiffSingleResult const& Result = *Item;
+	if (Result.Pin1)
+	{
+		PanelNew.GraphEditor.Pin()->ClearSelectionSet();
+		PanelOld.GraphEditor.Pin()->ClearSelectionSet();
+
+		if (Result.Pin1)
+		{
+			Result.Pin1->bIsDiffing = true;
+			GetGraphEditorForGraph(Result.Pin1->GetOwningNode()->GetGraph())->JumpToPin(Result.Pin1);
+		}
+
+		if (Result.Pin2)
+		{
+			Result.Pin2->bIsDiffing = true;
+			GetGraphEditorForGraph(Result.Pin2->GetOwningNode()->GetGraph())->JumpToPin(Result.Pin2);
+		}
+	}
+	else if (Result.Node1)
+	{
+		PanelNew.GraphEditor.Pin()->ClearSelectionSet();
+		PanelOld.GraphEditor.Pin()->ClearSelectionSet();
+
+		if (Result.Node2)
+		{
+			GetGraphEditorForGraph(Result.Node2->GetGraph())->JumpToNode(Result.Node2, false);
+		}
+		GetGraphEditorForGraph(Result.Node1->GetGraph())->JumpToNode(Result.Node1, false);
+	}
+}
+
+void SBlueprintMerge::ResetGraphEditors()
+{
+	auto GraphEditor = PanelLocal.GraphEditor.Pin();
+	if (GraphEditor != TSharedPtr<SGraphEditor>())
+	{
+		if (bLockViews)
+		{
+			GraphEditor->LockToGraphEditor(PanelNew.GraphEditor);
+			GraphEditor->LockToGraphEditor(PanelOld.GraphEditor);
+
+			auto OldGraph = PanelOld.GraphEditor.Pin();
+			if( OldGraph != TSharedPtr<SGraphEditor>() )
+			{
+				OldGraph->LockToGraphEditor( GraphEditor );
+			}
+			auto NewGraph = PanelNew.GraphEditor.Pin();
+			if( NewGraph!= TSharedPtr<SGraphEditor>() )
+			{
+				NewGraph->LockToGraphEditor(GraphEditor);
+			}
+		}
+		else
+		{
+			GraphEditor->UnlockFromGraphEditor(PanelNew.GraphEditor);
+			GraphEditor->UnlockFromGraphEditor(PanelOld.GraphEditor);
+
+			auto OldGraph = PanelOld.GraphEditor.Pin();
+			if (OldGraph != TSharedPtr<SGraphEditor>())
+			{
+				OldGraph->UnlockFromGraphEditor(GraphEditor);
+			}
+			auto NewGraph = PanelNew.GraphEditor.Pin();
+			if (NewGraph != TSharedPtr<SGraphEditor>())
+			{
+				NewGraph->UnlockFromGraphEditor(GraphEditor);
+			}
+		}
+	}
+	return SBlueprintDiff::ResetGraphEditors();
+}
+
+TSharedRef<SWidget> SBlueprintMerge::GenerateDiffWindow()
+{
+	return SNew(SVerticalBox)
+		+ SVerticalBox::Slot()
+		[
+			SNew(SSplitter).Orientation(Orient_Vertical)
+			+ SSplitter::Slot()
+			.Value(.7f)
+			[
+				SNew(SSplitter)
+				+ SSplitter::Slot()
+				.Value(0.5f)
+				[
+					//left blueprint
+					SAssignNew(PanelLocal.GraphEditorBorder, SBorder)
+					.VAlign(VAlign_Fill)
+					[
+						SBlueprintDiff::DefaultEmptyPanel()
+					]
+				]
+				+ SSplitter::Slot()
+				.Value(0.5f)
+				[
+					//middle (base) blueprint
+					SAssignNew(PanelOld.GraphEditorBorder, SBorder)
+					.VAlign(VAlign_Fill)
+					[
+						SBlueprintDiff::DefaultEmptyPanel()
+					]
+				]
+			+ SSplitter::Slot()
+				.Value(0.5f)
+				[
+					//right blueprint
+					SAssignNew(PanelNew.GraphEditorBorder, SBorder)
+					.VAlign(VAlign_Fill)
+					[
+						SBlueprintDiff::DefaultEmptyPanel()
+					]
+				]
+			]
+			+ SSplitter::Slot()
+			.Value(.3f)
+			[
+				SAssignNew(EditorBorder, SBorder)
+			]
+		]
+	;
+}
+
+TSharedRef<SWidget> SBlueprintMerge::GenerateToolbar()
+{
+	return SNew(SVerticalBox)
+		+ SVerticalBox::Slot()
+		.AutoHeight()
+		[
+			SNew(SVerticalBox)
+			+ SVerticalBox::Slot()
+			[
+				SNew(SButton)
+				.OnClicked(this, &SBlueprintMerge::OnAcceptResultClicked)
+				.Content()
+				[
+					SNew(STextBlock)
+					.Text(LOCTEXT("AcceptResult", "Accept Result"))
+					.ToolTipText( LOCTEXT( "AcceptResultToolTip", "Overwrite the local blueprint with the blueprint constructed in the lower portion of the merge tool." ) )
+				]
+			]
+			+ SVerticalBox::Slot()
+			[
+				SNew(SButton)
+				.OnClicked(this, &SBlueprintMerge::OnCancelClicked)
+				.Content()
+				[
+					SNew(STextBlock)
+					.Text(LOCTEXT("CancelMerge", "Cancel"))
+					.ToolTipText(LOCTEXT("CancelMergeToolTip", "Discards the changes made in the merge tool. Local file will remain unresolved so the merge tool can be run again if desired."))
+				]
+			]
+			+ SVerticalBox::Slot()
+			[
+				SNew(SButton)
+				.OnClicked(this, &SBlueprintMerge::OnTakeLocalClicked)
+				.Content()
+				[
+					SNew(STextBlock)
+					.Text(LOCTEXT("LocalChanges", "Take Local Changes Only"))
+					.ToolTipText(LOCTEXT("LocalChangesToolTip", "Resolve the conflict by discarding the changes made remotely. Conflict not fully resolved until the result is accepted."))
+				]
+			]
+			+ SVerticalBox::Slot()
+			[
+				SNew(SButton)
+				.OnClicked(this, &SBlueprintMerge::OnTakeRemoteClicked)
+				.Content()
+				[
+					SNew(STextBlock)
+					.Text(LOCTEXT("RemoteChanges", "Take Remote Changes Only"))
+					.ToolTipText(LOCTEXT("RemoteChangesToolTip", "Resolve the conflict by discarding the changes made locally. Conflict not fully resolved until the result is accepted."))
+				]
+			]
+			+ SVerticalBox::Slot()
+			[
+				SNew(SButton)
+				.OnClicked(this, &SBlueprintMerge::OnTakeBaseClicked)
+				.Content()
+				[
+					SNew(STextBlock)
+					.Text(LOCTEXT("BaseChanges", "Take Base Only (Revert All)"))
+					.ToolTipText(LOCTEXT("BaseChangesToolTip", "Resolve the conflict by discarding the both local and remote changes. Conflict not fully resolved until the result is accepted."))
+				]
+			]
+			+ SVerticalBox::Slot()
+			[
+				SNew(SButton)
+				.Content()
+				[
+					SNew(STextBlock)
+					.Text(LOCTEXT("AutoMerge", "Automatically Merge"))
+					.ToolTipText(LOCTEXT("AutoMergeChangesToolTip", "Resolve the conflict by attempting to automatically apply both local and remote changes. Conflict not fully resolved until the result is accepted."))
+				]
+			]
+		]
+		+ SVerticalBox::Slot()
+			.AutoHeight()
+		[
+			SAssignNew(DiffListBorder, SBorder)
+		];
+}
+
+// @cr doc const shallowness
+static UEdGraph* FindGraphByName(UBlueprint const& FromBlueprint, const FString& GraphName)
+{
+	TArray<UEdGraph*> Graphs;
+	FromBlueprint.GetAllGraphs(Graphs);
+
+	UEdGraph* Ret = NULL;
+	if (UEdGraph** Result = Graphs.FindByPredicate(FMatchName(GraphName)))
+	{
+		Ret = *Result;
+	}
+	return Ret;
+}
+
+void SBlueprintMerge::HandleGraphChanged(const FString& GraphName)
+{
+	UEdGraph* GraphOld = FindGraphByName(*PanelOld.Blueprint, GraphName);
+	UEdGraph* GraphNew = FindGraphByName(*PanelNew.Blueprint, GraphName);
+	UEdGraph* GraphLocal = FindGraphByName(*PanelLocal.Blueprint, GraphName);
+	UEdGraph* GraphResult = FindGraphByName(*BlueprintResult, GraphName);
+
+	PanelOld.GeneratePanel(GraphOld, NULL);
+	PanelNew.GeneratePanel(GraphNew, GraphOld);
+	PanelLocal.GeneratePanel(GraphLocal, GraphOld);
+
+	{
+		// Set up a normal editor window so that the user can tweak the result of the merge:
+		// @todo doc: need to copy the target blueprint and destructively edit that one, not 
+		// the base blueprint..
+		if (GraphResult)
+		{
+			SGraphEditor::FGraphEditorEvents InEvents;
+
+			//FGraphAppearanceInfo AppearanceInfo;
+			//AppearanceInfo.CornerText = LOCTEXT("AppearanceCornerText_BlueprintDif", "DIFF").ToString();
+
+			/*if (!GraphEditorCommands.IsValid())
+			{
+				GraphEditorCommands = MakeShareable(new FUICommandList());
+
+				GraphEditorCommands->MapAction(FGenericCommands::Get().Copy,
+					FExecuteAction::CreateRaw(this, &FDiffPanel::CopySelectedNodes),
+					FCanExecuteAction::CreateRaw(this, &FDiffPanel::CanCopyNodes)
+					);
+			}*/
+
+			auto Editor = SNew(SGraphEditor)
+				.GraphToEdit(GraphResult)
+				.TitleBarEnabledOnly(false)
+				//.Appearance(AppearanceInfo)
+				.GraphEvents(InEvents);
+			EditorBorder->SetContent(Editor);
+		}
+		else
+		{
+			EditorBorder->ClearContent();
+		}
+	}
+}
+
+template< typename T >
+static void CopyToShared( const TArray<T>& Source, TArray< TSharedPtr<T> >& Dest )
+{
+	check( Dest.Num() == 0 );
+	for( auto i : Source )
+	{
+		Dest.Add( MakeShareable(new T( i )) );
+	}
+}
+
+static void WriteBackup( UPackage& Package, const FString& Directory, const FString& Filename )
+{
+	if( GEditor )
+	{
+		const FString DestinationFilename = Directory + FString("/") + Filename;
+		FString OriginalFilename;
+		if( FPackageName::DoesPackageExist(Package.GetName(), NULL, &OriginalFilename) )
+		{
+			IFileManager::Get().Copy(*DestinationFilename, *OriginalFilename);
+		}
+	}
+}
+
+void SBlueprintMerge::OnSelectionChanged(SBlueprintDiff::FGraphToDiff Item, ESelectInfo::Type SelectionType)
+{
+	if (SelectionType != ESelectInfo::OnMouseClick || !Item.IsValid())
+	{
+		return;
+	}
+
+	FString GraphName = Item->GetGraphOld()->GetName();
+	DisablePinDiffFocus(); // @todo doc: Not sure what this does..
+	HandleGraphChanged(GraphName);
+
+	// Ensure that regenerated graphs are properly locked/unlocked:
+	ResetGraphEditors();
+
+	// We'll diff the other two graphs against this shared base:
+	UEdGraph* GraphBase = FindGraphByName(*PanelOld.Blueprint, GraphName);
+
+	// Diff local version vs base:
+	TArray<FDiffSingleResult> FoundDiffsLocal;
+	UEdGraph* GraphLocal = FindGraphByName(*PanelLocal.Blueprint, GraphName);
+	FGraphDiffControl::DiffGraphs(GraphLocal, GraphBase, FoundDiffsLocal);
+
+	// Diff remote version vs. base:
+	TArray<FDiffSingleResult> FoundDiffsRemote;
+	UEdGraph* GraphRemote = FindGraphByName(*PanelNew.Blueprint, GraphName);
+	FGraphDiffControl::DiffGraphs(GraphRemote, GraphBase, FoundDiffsRemote);
+
+	// Display results of the two diffs:
+	FMergeDifferencesListCommands const& Commands = FMergeDifferencesListCommands::Get();
+	LocalDiffResults.Empty();
+	TSharedRef< SWidget > DiffsLocal = GenerateDiffView(FoundDiffsLocal, Commands.LocalNext, Commands.LocalPrevious, LocalDiffResults );
+	RemoteDiffResults.Empty();
+	TSharedRef< SWidget > DiffsRemote = GenerateDiffView(FoundDiffsRemote, Commands.RemoteNext, Commands.RemotePrevious, RemoteDiffResults );
+
+	DiffListBorder->SetContent( SNew( SVerticalBox )
+		+ SVerticalBox::Slot()[DiffsLocal]
+		+ SVerticalBox::Slot()[DiffsRemote] );
+}
+
+SGraphEditor* SBlueprintMerge::GetGraphEditorForGraph(UEdGraph* Graph) const
+{
+	if (PanelLocal.GraphEditor.Pin()->GetCurrentGraph() == Graph)
+	{
+		return PanelLocal.GraphEditor.Pin().Get();
+	}
+	return SBlueprintDiff::GetGraphEditorForGraph(Graph);
+}
+
+TSharedRef< SWidget > SBlueprintMerge::GenerateDiffView(TArray<FDiffSingleResult>& Diffs, TSharedPtr< const FUICommandInfo > CommandNext, TSharedPtr< const FUICommandInfo > CommandPrev, TArray< TSharedPtr<FDiffSingleResult> >& SharedResults )
+{
+	struct SortDiff
+	{
+		bool operator () (const FDiffSingleResult& A, const FDiffSingleResult& B) const
+		{
+			return A.Diff < B.Diff;
+		}
+	};
+
+	Sort(Diffs.GetTypedData(), Diffs.Num(), SortDiff());
+
+	TSharedPtr<FUICommandList> KeyCommands = MakeShareable(new FUICommandList);
+
+	//KeyCommands->MapAction(CommandPrev, FExecuteAction::CreateSP(this, &FListItemGraphToDiff::PrevDiff));
+	//KeyCommands->MapAction(CommandNext, FExecuteAction::CreateSP(this, &FListItemGraphToDiff::NextDiff));
+
+	FToolBarBuilder ToolbarBuilder(KeyCommands, FMultiBoxCustomization::None);
+	ToolbarBuilder.AddToolBarButton(CommandPrev, NAME_None, TAttribute<FText>(), TAttribute<FText>(), FSlateIcon(FEditorStyle::GetStyleSetName(), "BlueprintDif.PrevDiff"));
+	ToolbarBuilder.AddToolBarButton(CommandNext, NAME_None, TAttribute<FText>(), TAttribute<FText>(), FSlateIcon(FEditorStyle::GetStyleSetName(), "BlueprintDif.NextDiff"));
+
+	CopyToShared(Diffs, SharedResults);
+
+	TSharedRef<SHorizontalBox> Result = SNew(SHorizontalBox)
+		+ SHorizontalBox::Slot()
+		.FillWidth(1.f)
+		.MaxWidth(350.f)
+		[
+			SNew(SVerticalBox)
+			+ SVerticalBox::Slot()
+			.Padding(0.f)
+			.AutoHeight()
+			[
+				ToolbarBuilder.MakeWidget()
+			]
+			+ SVerticalBox::Slot()
+				.Padding(0.f)
+				.AutoHeight()
+				[
+					SNew(SBorder)
+					.BorderImage(FEditorStyle::GetBrush("PropertyWindow.CategoryBackground"))
+					.Padding(FMargin(2.0f))
+					.ForegroundColor(FEditorStyle::GetColor("PropertyWindow.CategoryForeground"))
+					.ToolTipText(LOCTEXT("BlueprintDifDifferencesToolTip", "List of differences found between revisions, click to select"))
+					.HAlign(HAlign_Center)
+					[
+						SNew(STextBlock)
+						.Text(LOCTEXT("RevisionDifferences", "Revision Differences"))
+					]
+				]
+			+ SVerticalBox::Slot()
+				.Padding(1.f)
+				.FillHeight(1.f)
+				[
+					SNew(SListView< TSharedPtr<FDiffSingleResult> >)
+					.ItemHeight(24)
+					.ListItemsSource(&SharedResults)
+					.OnGenerateRow(this, &SBlueprintMerge::OnGenerateRow)
+					.SelectionMode(ESelectionMode::Single)
+					.OnSelectionChanged( this, &SBlueprintMerge::OnDiffListSelectionChanged)
+				]
+		];
+	return Result;
+}
+
+FReply SBlueprintMerge::OnAcceptResultClicked()
+{
+	// Because merge operations are so destructive and can be confusing, lets write backups of the files involved:
+	const FString BackupSubDir = FPaths::GameSavedDir() / TEXT("Backup") / TEXT("Resolve_Backup[") + FDateTime::Now().ToString(TEXT("%Y-%m-%d-%H-%M-%S")) + TEXT("]");
+	WriteBackup(*PanelNew.Blueprint->GetOutermost(), BackupSubDir, TEXT("RemoteAsset") + FPackageName::GetAssetPackageExtension() );
+	WriteBackup(*PanelOld.Blueprint->GetOutermost(), BackupSubDir, TEXT("CommonBaseAsset") + FPackageName::GetAssetPackageExtension() );
+	WriteBackup(*PanelLocal.Blueprint->GetOutermost(), BackupSubDir, TEXT("LocalAsset") + FPackageName::GetAssetPackageExtension());
+
+	UPackage* Package = PanelLocal.Blueprint->GetOutermost();
+	TArray<UPackage*> PackagesToSave;
+	PackagesToSave.Add(Package);
+
+	// Perform the resolve with the SCC plugin, we do this first so that the editor doesn't warn about writing to a file that is unresolved:
+	ISourceControlModule::Get().GetProvider().Execute(ISourceControlOperation::Create<FResolve>(), SourceControlHelpers::PackageFilenames(PackagesToSave), EConcurrency::Synchronous);
+
+	// cr doc: this is too simple to be correct..
+	PanelLocal.Blueprint->RemoveGeneratedClasses();
+	StaticDuplicateObject( BlueprintResult, Package, *PanelLocal.Blueprint->GetName() );
+
+	auto OwningWindowPinned = OwningWindow.Pin();
+	if (OwningWindowPinned.IsValid())
+	{
+		OwningWindowPinned->RequestDestroyWindow();
+	}
+
+	const auto Result = FEditorFileUtils::PromptForCheckoutAndSave(PackagesToSave, /*bCheckDirty=*/ false, /*bPromptToSave=*/ false);
+	if (Result != FEditorFileUtils::PR_Success)
+	{
+		static const float MessageTime = 5.0f;
+		const FText ErrorMessage = FText::Format( LOCTEXT("MergeWriteFailedError", "Failed to write merged files, please look for backups in {0}"), FText::FromString(BackupSubDir) );
+
+		FNotificationInfo ErrorNotification(ErrorMessage);
+		FSlateNotificationManager::Get().AddNotification( ErrorNotification );
+	}
+
+	return FReply::Handled();
+}
+
+FReply SBlueprintMerge::OnCancelClicked()
+{
+	auto OwningWindowPinned = OwningWindow.Pin();
+	if( OwningWindowPinned.IsValid() )
+	{
+		OwningWindowPinned->RequestDestroyWindow();
+	}
+	return FReply::Handled();
+}
+
+FReply SBlueprintMerge::OnTakeLocalClicked()
+{
+	StageBlueprint(PanelLocal.Blueprint);
+	return FReply::Handled();
+}
+
+FReply SBlueprintMerge::OnTakeRemoteClicked()
+{
+	StageBlueprint(PanelNew.Blueprint);
+	return FReply::Handled();
+}
+
+FReply SBlueprintMerge::OnTakeBaseClicked()
+{
+	StageBlueprint(PanelOld.Blueprint);
+	return FReply::Handled();
+}
+
+void SBlueprintMerge::StageBlueprint(UBlueprint const* DesiredBP)
+{
+	// User has decided to discard the remote changes, update the result window
+	// with this result:
+	if (BlueprintResult)
+	{
+		//delete BlueprintResult;... yeah.. i have no idea what's keeping this stuff alive, best not to delete it
+	}
+
+	BlueprintResult = DuplicateBlueprint(DesiredBP);
+
+	(void)PackagesAppearEquivalent(BlueprintResult->GetOutermost(), DesiredBP->GetOutermost());
+	//check( PackagesAppearEquivalent( BlueprintResult->GetOutermost(), DesiredBP->GetOutermost() ) );
+
+	SGraphEditor::FGraphEditorEvents InEvents;
+	TArray< FGraphToDiff>  Selected = GraphsToDiff->GetSelectedItems();
+	check(Selected.Num() <= 1); // We set this control up to not support multiselect
+
+	UEdGraph* GraphResult = NULL;
+
+	if (Selected.Num() == 1)
+	{
+		GraphResult = FindGraphByName(*BlueprintResult, Selected[0]->GetGraphNew()->GetName());
+	}
+
+	if( GraphResult )
+	{
+		auto Editor = SNew(SGraphEditor)
+			.GraphToEdit(GraphResult)
+			.TitleBarEnabledOnly(false)
+			//.Appearance(AppearanceInfo)
+			.GraphEvents(InEvents);
+		EditorBorder->SetContent(Editor);
+	}
+	else
+	{
+		EditorBorder->ClearContent();
+	}
+}
+
+#undef LOCTEXT_NAMESPACE
diff --git a/Engine/Source/Developer/Merge/Private/SBlueprintMerge.h b/Engine/Source/Developer/Merge/Private/SBlueprintMerge.h
new file mode 100644
--- /dev/null
+++ b/Engine/Source/Developer/Merge/Private/SBlueprintMerge.h
@@ -0,0 +1,50 @@
+// Copyright 1998-2014 Epic Games, Inc. All Rights Reserved.
+#pragma once
+
+#include "SBlueprintDiff.h"
+
+class MERGE_API SBlueprintMerge : public SBlueprintDiff
+{
+public:
+	SLATE_BEGIN_ARGS(SBlueprintMerge){}
+	SLATE_ARGUMENT(class UBlueprint*, BlueprintLocal)
+	SLATE_ARGUMENT( SBlueprintDiff::FArguments, BaseArgs )
+	SLATE_ARGUMENT( TWeakPtr<SWindow>, OwningWindow )
+	SLATE_END_ARGS()
+	
+	void Construct(const FArguments& InArgs);
+	TSharedRef<ITableRow> OnGenerateRow( TSharedPtr<struct FDiffSingleResult> ParamItem, const TSharedRef<STableViewBase>& OwnerTable);
+	void OnDiffListSelectionChanged(TSharedPtr<struct FDiffSingleResult> Item, ESelectInfo::Type SelectionType);
+
+protected:
+	/** Overrides, see base class for documentation */
+	virtual void ResetGraphEditors() override;
+	virtual TSharedRef<SWidget> GenerateDiffWindow() override;
+	virtual TSharedRef<SWidget> GenerateToolbar() override;
+	virtual void HandleGraphChanged(const FString& GraphName) override;
+	virtual void OnSelectionChanged(FGraphToDiff Item, ESelectInfo::Type SelectionType) override;
+	virtual SGraphEditor* GetGraphEditorForGraph(UEdGraph* Graph) const override;
+	
+	/** Helper functions */
+	TSharedRef< SWidget > GenerateDiffView(TArray<FDiffSingleResult>& Diffs, TSharedPtr< const FUICommandInfo > CommandNext, TSharedPtr< const FUICommandInfo > CommandPrev, TArray< TSharedPtr<FDiffSingleResult> >& SharedResults);
+
+	/** Event handlers */
+	FReply OnAcceptResultClicked();
+	FReply OnCancelClicked();
+	FReply OnTakeLocalClicked();
+	FReply OnTakeRemoteClicked();
+	FReply OnTakeBaseClicked();
+	void StageBlueprint(UBlueprint const* DesiredBP);
+
+	// Raw pointer?!
+	UBlueprint* BlueprintResult;
+	FDiffPanel PanelLocal;
+	TSharedPtr<SBorder>		 EditorBorder;
+	TWeakPtr<SWindow>		OwningWindow;
+
+	// This has to be allocated here because SListView cannot own the list
+	// that it is displaying. It also seems like the display list *has*
+	// to be a list of TSharedPtrs.
+	TArray< TSharedPtr<FDiffSingleResult> > LocalDiffResults;
+	TArray< TSharedPtr<FDiffSingleResult> > RemoteDiffResults;
+};
diff --git a/Engine/Source/Developer/Merge/Public/Merge.h b/Engine/Source/Developer/Merge/Public/Merge.h
new file mode 100644
--- /dev/null
+++ b/Engine/Source/Developer/Merge/Public/Merge.h
@@ -0,0 +1,53 @@
+// Copyright 1998-2014 Epic Games, Inc. All Rights Reserved.
+
+#pragma once
+
+#include "ModuleManager.h"
+#include "IAssetTypeActions.h"
+
+class UBlueprint;
+
+struct FMergeDisplayArgs
+{
+	FRevisionInfo DepotHeadRevision;
+	FRevisionInfo CommonBaseRevision;
+};
+
+/**
+ * The public interface to this module
+ */
+class IMerge : public IModuleInterface
+{
+
+public:
+
+	/**
+	 * Singleton-like access to this module's interface.  This is just for convenience!
+	 * Beware of calling this during the shutdown phase, though.  Your module might have been unloaded already.
+	 *
+	 * @return Returns singleton instance, loading the module on demand if needed
+	 */
+	static inline IMerge& Get()
+	{
+		return FModuleManager::LoadModuleChecked< IMerge >( "Merge" );
+	}
+
+	/**
+	 * Checks to see if this module is loaded and ready.  It is only valid to call Get() if IsAvailable() returns true.
+	 *
+	 * @return True if the module is loaded and ready to use
+	 */
+	static inline bool IsAvailable()
+	{
+		return FModuleManager::Get().IsModuleLoaded( "Merge" );
+	}
+
+	/**
+	 * Performs a three way merge of AssetA and AssetB using CommonOrigin as a base. Takes a display args structure
+	 * to customize the view.
+	 *
+	 * @return The resulting UBlueprint, NULL if merge was canceled or otherwise failed
+	 */
+	virtual void Merge( const UBlueprint& Object ) = 0;
+};
+
diff --git a/Engine/Source/Developer/SourceControl/Public/SourceControlOperations.h b/Engine/Source/Developer/SourceControl/Public/SourceControlOperations.h
--- a/Engine/Source/Developer/SourceControl/Public/SourceControlOperations.h
+++ b/Engine/Source/Developer/SourceControl/Public/SourceControlOperations.h
@@ -292,4 +292,22 @@ protected:
 	FString Destination;
 };
 
+/**
+ * Operation used to resolve a file that is in a conflicted state.
+ */
+class FResolve : public ISourceControlOperation
+{
+public:
+	// ISourceControlOperation interface
+	virtual FName GetName() const override
+	{
+		  return "Resolve";
+	}
+
+	virtual FText GetInProgressString() const override
+	{
+		return LOCTEXT("SourceControl_Resolve", "Resolving file(s) in Source Control...");
+	}
+};
+
 #undef LOCTEXT_NAMESPACE
\ No newline at end of file
diff --git a/Engine/Source/Editor/ContentBrowser/ContentBrowser.Build.cs b/Engine/Source/Editor/ContentBrowser/ContentBrowser.Build.cs
--- a/Engine/Source/Editor/ContentBrowser/ContentBrowser.Build.cs
+++ b/Engine/Source/Editor/ContentBrowser/ContentBrowser.Build.cs
@@ -24,6 +24,7 @@ public class ContentBrowser : ModuleRules
 				"CoreUObject",
                 "InputCore",
 				"Engine",
+                "Merge",
 				"Slate",
 				"SlateCore",
                 "EditorStyle",
diff --git a/Engine/Source/Editor/ContentBrowser/Private/AssetContextMenu.cpp b/Engine/Source/Editor/ContentBrowser/Private/AssetContextMenu.cpp
--- a/Engine/Source/Editor/ContentBrowser/Private/AssetContextMenu.cpp
+++ b/Engine/Source/Editor/ContentBrowser/Private/AssetContextMenu.cpp
@@ -10,6 +10,7 @@
 #include "Editor/UnrealEd/Public/FileHelpers.h"
 #include "Editor/PropertyEditor/Public/PropertyEditorModule.h"
 #include "Toolkits/AssetEditorManager.h"
+#include "Toolkits/ToolkitManager.h"
 #include "ConsolidateWindow.h"
 #include "ReferenceViewer.h"
 
@@ -21,6 +22,7 @@
 #include "KismetEditorUtilities.h"
 #include "AssetToolsModule.h"
 #include "ComponentAssetBroker.h"
+#include "Merge.h"
 
 #define LOCTEXT_NAMESPACE "ContentBrowser"
 
@@ -31,6 +33,7 @@ FAssetContextMenu::FAssetContextMenu(const TWeakPtr<SAssetView>& InAssetView)
 	, bCanExecuteSCCCheckIn(false)
 	, bCanExecuteSCCHistory(false)
 	, bCanExecuteSCCRevert(false)
+	, bCanExecuteSCCMerge(false)
 	, bCanExecuteSCCSync(false)
 	, AssetView(InAssetView)
 {
@@ -404,6 +407,19 @@ bool FAssetContextMenu::AddSourceControlMenuOptions(FMenuBuilder& MenuBuilder)
 				);
 		}
 
+		if (CanExecuteSCCMerge())
+		{
+			MenuBuilder.AddMenuEntry(
+				LOCTEXT("SCCMerge", "Merge"),
+				LOCTEXT("SCCMergeTooltip", "Initiates a three-way merge of an asset."),
+				FSlateIcon(),
+				FUIAction(
+				FExecuteAction::CreateSP(this, &FAssetContextMenu::ExecuteSCCMerge),
+				FCanExecuteAction::CreateSP(this, &FAssetContextMenu::CanExecuteSCCMerge)
+				)
+				);
+		}
+
 		MenuBuilder.AddMenuEntry(
 			LOCTEXT("SCCRefresh", "Refresh"),
 			LOCTEXT("SCCRefreshTooltip", "Updates the source control status of the asset."),
@@ -449,7 +465,6 @@ bool FAssetContextMenu::AddSourceControlMenuOptions(FMenuBuilder& MenuBuilder)
 					)
 				);
 		}
-
 	}
 	else
 	{
@@ -1107,6 +1122,24 @@ void FAssetContextMenu::ExecuteSCCRevert()
 	FSourceControlWindows::PromptForRevert(PackageNames);
 }
 
+void FAssetContextMenu::ExecuteSCCMerge()
+{
+	TArray<FString> PackageNames;
+	GetSelectedPackageNames(PackageNames);
+
+	for (int32 AssetIdx = 0; AssetIdx < SelectedAssets.Num(); ++AssetIdx)
+	{
+		const FAssetData& Asset = SelectedAssets[AssetIdx];
+		UObject* Object = Asset.GetAsset();
+
+		if( Object )
+		{
+			IMerge::Get().Merge(*Cast<UBlueprint>(Object));
+		}
+	}
+
+}
+
 void FAssetContextMenu::ExecuteSCCSync()
 {
 	TArray<FString> PackageNames;
@@ -1237,6 +1270,13 @@ bool FAssetContextMenu::CanExecuteSCCRevert() const
 	return bCanExecuteSCCRevert;
 }
 
+bool FAssetContextMenu::CanExecuteSCCMerge() const
+{
+	bool bEnableMerge = false;
+	GConfig->GetBool(TEXT("AssetMerge"), TEXT("EnableAssetMerge"), bEnableMerge, GEngineIni);
+	return bCanExecuteSCCMerge && bEnableMerge;
+}
+
 bool FAssetContextMenu::CanExecuteSCCSync() const
 {
 	return bCanExecuteSCCSync;
@@ -1332,6 +1372,7 @@ void FAssetContextMenu::CacheCanExecuteVars()
 	bCanExecuteSCCCheckIn = false;
 	bCanExecuteSCCHistory = false;
 	bCanExecuteSCCRevert = false;
+	bCanExecuteSCCMerge = false;
 	bCanExecuteSCCSync = false;
 
 	for (auto AssetIt = SelectedAssets.CreateConstIterator(); AssetIt; ++AssetIt)
@@ -1378,6 +1419,11 @@ void FAssetContextMenu::CacheCanExecuteVars()
 					bCanExecuteSCCCheckIn = true;
 					bCanExecuteSCCRevert = true;
 				}
+
+				if( SourceControlState->IsConflicted() && SourceControlState->IsCurrent() )
+				{
+					bCanExecuteSCCMerge = true;
+				}
 			}
 		}
 
diff --git a/Engine/Source/Editor/ContentBrowser/Private/AssetContextMenu.h b/Engine/Source/Editor/ContentBrowser/Private/AssetContextMenu.h
--- a/Engine/Source/Editor/ContentBrowser/Private/AssetContextMenu.h
+++ b/Engine/Source/Editor/ContentBrowser/Private/AssetContextMenu.h
@@ -133,6 +133,9 @@ private:
 	/** Handler for when "Source Control Revert" is selected */
 	void ExecuteSCCRevert();
 
+	/** Handler for when "Source Control Merge" is selected */
+	void ExecuteSCCMerge();
+
 	/** Handler for when "Source Control Sync" is selected */
 	void ExecuteSCCSync();
 
@@ -187,6 +190,9 @@ private:
 	/** Handler to check to see if "Source Control Revert" can be executed */
 	bool CanExecuteSCCRevert() const;
 
+	/** Handler to check to see if "Source Control Merge" can be executed */
+	bool CanExecuteSCCMerge() const;
+
 	/** Handler to check to see if "Source Control Sync" can be executed */
 	bool CanExecuteSCCSync() const;
 
@@ -246,5 +252,6 @@ private:
 	bool bCanExecuteSCCCheckIn;
 	bool bCanExecuteSCCHistory;
 	bool bCanExecuteSCCRevert;
+	bool bCanExecuteSCCMerge;
 	bool bCanExecuteSCCSync;
 };
diff --git a/Engine/Source/Editor/Kismet/Private/SBlueprintDiff.cpp b/Engine/Source/Editor/Kismet/Private/SBlueprintDiff.cpp
--- a/Engine/Source/Editor/Kismet/Private/SBlueprintDiff.cpp
+++ b/Engine/Source/Editor/Kismet/Private/SBlueprintDiff.cpp
@@ -328,7 +328,7 @@ void FListItemGraphToDiff::KeyWasPressed( const FKeyboardEvent& InKeyboardEvent
 }
 
 
-SBlueprintDiff::FDiffPanel::FDiffPanel()
+FDiffPanel::FDiffPanel()
 {
 	Blueprint = NULL;
 }
@@ -392,14 +392,6 @@ void SBlueprintDiff::Construct( const FArguments& InArgs)
 			CreateGraphEntry(NULL,GraphNew);
 		}
 	}
-	auto DefaultEmptyPanel = SNew(SHorizontalBox)
-		+SHorizontalBox::Slot()
-		.HAlign(HAlign_Center)
-		.VAlign(VAlign_Center)
-		[
-			SNew(STextBlock)
-			.Text(LOCTEXT("BlueprintDifGraphsToolTip", "Select Graph to Diff"))
-		];
 
 	this->ChildSlot
 		[	
@@ -476,24 +468,7 @@ void SBlueprintDiff::Construct( const FArguments& InArgs)
 								+SSplitter::Slot()
 								.Value(0.7f)
 								[
-									SNew(SVerticalBox)
-									+SVerticalBox::Slot()
-									.AutoHeight()
-									[
-										//open in p4dif tool
-										SNew(SButton)
-										.OnClicked(this, &SBlueprintDiff::OnOpenInDefaults)
-										.Content()
-										[
-											SNew(STextBlock)
-											.Text(LOCTEXT("DifBlueprintDefaultsToolTip", "Diff Blueprint Defaults"))
-										]
-									]
-									+SVerticalBox::Slot()
-									.FillHeight(1.f)
-									[
-										SAssignNew(DiffListBorder, SBorder)
-									]
+									GenerateToolbar()
 								]
 							]
 						]
@@ -504,28 +479,8 @@ void SBlueprintDiff::Construct( const FArguments& InArgs)
 							+SHorizontalBox::Slot()
 							.FillWidth(1.f)
 							[
-								//dif window
-								SNew(SSplitter)
-								+SSplitter::Slot()
-								.Value(0.5f)
-								[
-									//left blueprint
-									SAssignNew(PanelOld.GraphEditorBorder, SBorder)
-									.VAlign(VAlign_Fill)
-									[
-										DefaultEmptyPanel
-									]
-								]
-								+SSplitter::Slot()
-								.Value(0.5f)
-								[
-									//right blueprint
-									SAssignNew(PanelNew.GraphEditorBorder, SBorder)
-									.VAlign(VAlign_Fill)
-									[
-										DefaultEmptyPanel
-									]
-								]
+								//diff window
+								GenerateDiffWindow()
 							]
 						]
 					]
@@ -572,8 +527,7 @@ void SBlueprintDiff::OnGraphChanged(FListItemGraphToDiff* Diff)
 void SBlueprintDiff::FocusOnGraphRevisions( class UEdGraph* GraphOld, class UEdGraph* GraphNew, FListItemGraphToDiff* Diff )
 {
 	DisablePinDiffFocus();
-	PanelOld.GeneratePanel(GraphOld,GraphNew);
-	PanelNew.GeneratePanel(GraphNew, GraphOld);
+	HandleGraphChanged( GraphOld->GetName() );
 
 	ResetGraphEditors();
 
@@ -649,7 +603,7 @@ void SBlueprintDiff::ResetGraphEditors()
 	}
 }
 
-void SBlueprintDiff::FDiffPanel::GeneratePanel(UEdGraph* Graph, UEdGraph* GraphToDiff)
+void FDiffPanel::GeneratePanel(UEdGraph* Graph, UEdGraph* GraphToDiff)
 {
 	TSharedPtr<SWidget> Widget = SNew(SBorder)
 								.HAlign(HAlign_Center)
@@ -694,7 +648,7 @@ void SBlueprintDiff::FDiffPanel::GeneratePanel(UEdGraph* Graph, UEdGraph* GraphT
 	GraphEditorBorder->SetContent(Widget.ToSharedRef());
 }
 
-FGraphPanelSelectionSet SBlueprintDiff::FDiffPanel::GetSelectedNodes() const
+FGraphPanelSelectionSet FDiffPanel::GetSelectedNodes() const
 {
 	FGraphPanelSelectionSet CurrentSelection;
 	TSharedPtr<SGraphEditor> FocusedGraphEd = GraphEditor.Pin();
@@ -705,7 +659,7 @@ FGraphPanelSelectionSet SBlueprintDiff::FDiffPanel::GetSelectedNodes() const
 	return CurrentSelection;
 }
 
-void SBlueprintDiff::FDiffPanel::CopySelectedNodes()
+void FDiffPanel::CopySelectedNodes()
 {
 	// Export the selected nodes and place the text on the clipboard
 	const FGraphPanelSelectionSet SelectedNodes = GetSelectedNodes();
@@ -716,7 +670,7 @@ void SBlueprintDiff::FDiffPanel::CopySelectedNodes()
 }
 
 
-FString SBlueprintDiff::FDiffPanel::GetTitle() const
+FString FDiffPanel::GetTitle() const
 {
 	FString Title = LOCTEXT("CurrentRevision", "Current Revision").ToString();
 
@@ -748,7 +702,7 @@ FString SBlueprintDiff::FDiffPanel::GetTitle() const
 	return Title;
 }
 
-bool SBlueprintDiff::FDiffPanel::CanCopyNodes() const
+bool FDiffPanel::CanCopyNodes() const
 {
 	// If any of the nodes can be duplicated then we should allow copying
 	const FGraphPanelSelectionSet SelectedNodes = GetSelectedNodes();
@@ -783,6 +737,86 @@ SGraphEditor* SBlueprintDiff::GetGraphEditorForGraph( UEdGraph* Graph ) const
 	return NULL;
 }
 
+TSharedRef<SWidget> SBlueprintDiff::DefaultEmptyPanel()
+{
+	return SNew(SHorizontalBox)
+		+ SHorizontalBox::Slot()
+		.HAlign(HAlign_Center)
+		.VAlign(VAlign_Center)
+		[
+			SNew(STextBlock)
+			.Text(LOCTEXT("BlueprintDifGraphsToolTip", "Select Graph to Diff"))
+		];
+}
+
+TSharedRef<SWidget> SBlueprintDiff::GenerateDiffWindow()
+{
+	return SNew(SSplitter)
+		+ SSplitter::Slot()
+		.Value(0.5f)
+		[
+			//left blueprint
+			SAssignNew(PanelOld.GraphEditorBorder, SBorder)
+			.VAlign(VAlign_Fill)
+			[
+				DefaultEmptyPanel()
+			]
+		]
+	+ SSplitter::Slot()
+		.Value(0.5f)
+		[
+			//right blueprint
+			SAssignNew(PanelNew.GraphEditorBorder, SBorder)
+			.VAlign(VAlign_Fill)
+			[
+				DefaultEmptyPanel()
+			]
+		];
+}
+
+TSharedRef<SWidget> SBlueprintDiff::GenerateToolbar()
+{
+	return SNew(SVerticalBox)
+		+SVerticalBox::Slot()
+		.AutoHeight()
+		[
+			//open in p4dif tool
+			SNew(SButton)
+			.OnClicked(this, &SBlueprintDiff::OnOpenInDefaults)
+			.Content()
+			[
+				SNew(STextBlock)
+				.Text(LOCTEXT("DifBlueprintDefaultsToolTip", "Diff Blueprint Defaults"))
+			]
+		]
+		+SVerticalBox::Slot()
+		.FillHeight(1.f)
+		[
+			SAssignNew(DiffListBorder, SBorder)
+		];
+}
+
+void SBlueprintDiff::HandleGraphChanged( const FString& GraphName )
+{
+	TArray<UEdGraph*> GraphsOld, GraphsNew;
+	PanelOld.Blueprint->GetAllGraphs(GraphsOld);
+	PanelNew.Blueprint->GetAllGraphs(GraphsNew);
+
+	UEdGraph* GraphOld = NULL;
+	if( UEdGraph** Iter = GraphsOld.FindByPredicate(FMatchName(GraphName)) )
+	{
+		GraphOld = *Iter;
+	}
+	UEdGraph* GraphNew = NULL;
+	if (UEdGraph** Iter = GraphsNew.FindByPredicate(FMatchName(GraphName)))
+	{
+		GraphNew = *Iter;
+	}
+
+	PanelOld.GeneratePanel(GraphOld, GraphNew);
+	PanelNew.GeneratePanel(GraphNew, GraphOld);
+}
+
 FReply SBlueprintDiff::OnKeyDown( const FGeometry& MyGeometry, const FKeyboardEvent& InKeyboardEvent ) 
 {
 	TArray< FGraphToDiff>  Selected = GraphsToDiff->GetSelectedItems();
diff --git a/Engine/Source/Editor/Kismet/Public/SBlueprintDiff.h b/Engine/Source/Editor/Kismet/Public/SBlueprintDiff.h
--- a/Engine/Source/Editor/Kismet/Public/SBlueprintDiff.h
+++ b/Engine/Source/Editor/Kismet/Public/SBlueprintDiff.h
@@ -4,6 +4,21 @@
 #include "Developer/AssetTools/Public/IAssetTypeActions.h"
 #include "GraphEditor.h"
 
+struct FMatchName
+{
+	FMatchName(const FString& InName)
+	: Name(InName)
+	{
+	}
+
+	bool operator() (const UObject* Object)
+	{
+		return Object->GetName() == Name;
+	}
+
+	FString const& Name;
+};
+
 /*List item that entry for a graph*/
 struct FListItemGraphToDiff: public TSharedFromThis<FListItemGraphToDiff>
 {
@@ -87,6 +102,46 @@ private:
 	TSharedPtr< FUICommandList > KeyCommands;
 };
 
+/*panel used to display the blueprint*/
+struct KISMET_API FDiffPanel
+{
+	FDiffPanel();
+
+	/*Generate this panel based on the specified graph*/
+	void	GeneratePanel(UEdGraph* Graph, UEdGraph* GraphToDiff);
+
+	/*Get the title to show at the top of the panel*/
+	FString GetTitle() const;
+
+	/* Called when user hits keyboard shortcut to copy nodes*/
+	void CopySelectedNodes();
+
+	/*Gets whatever nodes are selected in the Graph Editor*/
+	FGraphPanelSelectionSet GetSelectedNodes() const;
+
+	/*Can user copy any of the selected nodes?*/
+	bool CanCopyNodes() const;
+
+	/*The blueprint that owns the graph we are showing*/
+	class UBlueprint*				Blueprint;
+
+	/*The border around the graph editor, used to change the content when new graphs are set */
+	TSharedPtr<SBorder>				GraphEditorBorder;
+
+	/*The graph editor which does the work of displaying the graph*/
+	TWeakPtr<class SGraphEditor>	GraphEditor;
+
+	/*Revision information for this blueprint */
+	FRevisionInfo					RevisionInfo;
+
+	/*A name identifying which asset this panel is displaying */
+	bool							bShowAssetName;
+
+private:
+	/*Command list for this diff panel*/
+	TSharedPtr<FUICommandList> GraphEditorCommands;
+};
+
 /* Visual Diff between two Blueprints*/
 class  KISMET_API SBlueprintDiff: public SCompoundWidget
 {
@@ -108,7 +163,7 @@ public:
 	/** Called when a new Graph is clicked on by user */
 	void OnGraphChanged(FListItemGraphToDiff* Diff);
 
-private:
+protected:
 	/* Need to process keys for shortcuts to buttons */
 	virtual FReply OnKeyDown( const FGeometry& MyGeometry, const FKeyboardEvent& InKeyboardEvent ) override;
 
@@ -126,7 +181,7 @@ private:
 	TSharedRef<ITableRow> OnGenerateRow(FGraphToDiff ParamItem, const TSharedRef<STableViewBase>& OwnerTable );
 
 	/*Called when user clicks on a new graph list item */
-	void OnSelectionChanged(FGraphToDiff Item, ESelectInfo::Type SelectionType);
+	virtual void OnSelectionChanged(FGraphToDiff Item, ESelectInfo::Type SelectionType);
 
 	void OnDiffListSelectionChanged(const TSharedPtr<struct FDiffResultItem>& TheDiff, FListItemGraphToDiff* GraphDiffer);
 
@@ -143,7 +198,7 @@ private:
 	FReply  OnOpenInDefaults();
 
 	/*Reset the graph editor, called when user switches graphs to display*/
-	void    ResetGraphEditors();
+	virtual void ResetGraphEditors();
 
 	/*Get the image to show for the toggle lock option*/
 	const FSlateBrush* GetLockViewImage() const;
@@ -152,47 +207,19 @@ private:
 	TArray< FGraphToDiff> Graphs;
 
 	/** Get Graph editor associated with this Graph */
-	SGraphEditor* GetGraphEditorForGraph(UEdGraph* Graph) const ;
-
-	/*panel used to display the blueprint*/
-	struct FDiffPanel
-	{
-		FDiffPanel();
-
-		/*Generate this panel based on the specified graph*/
-		void	GeneratePanel(UEdGraph* Graph, UEdGraph* GraphToDiff);
-
-		/*Get the title to show at the top of the panel*/
-		FString GetTitle() const;
-
-		/* Called when user hits keyboard shortcut to copy nodes*/
-		void CopySelectedNodes();
-
-		/*Gets whatever nodes are selected in the Graph Editor*/
-		FGraphPanelSelectionSet GetSelectedNodes() const;
-
-		/*Can user copy any of the selected nodes?*/
-		bool CanCopyNodes() const;
-
-		/*The blueprint that owns the graph we are showing*/
-		class UBlueprint*				Blueprint;
-
-		/*The border around the graph editor, used to change the content when new graphs are set */
-		TSharedPtr<SBorder>				GraphEditorBorder;
+	virtual SGraphEditor* GetGraphEditorForGraph(UEdGraph* Graph) const ;
 
-		/*The graph editor which does the work of displaying the graph*/
-		TWeakPtr<class SGraphEditor>	GraphEditor;
+	/** Helper function for generating an empty widget */
+	static TSharedRef<SWidget> DefaultEmptyPanel();
 
-		/*Revision information for this blueprint */
-		FRevisionInfo					RevisionInfo;
+	/** Generates the widgets used to display the graphs that are being diff'd */
+	virtual TSharedRef<SWidget> GenerateDiffWindow();
 
-		/*A name identifying which asset this panel is displaying */
-		bool							bShowAssetName;
+	/** Generates the toolbar of buttons (for example, the button that lets you diff the CDO) */
+	virtual TSharedRef<SWidget> GenerateToolbar();
 
-	private:
-		/*Command list for this diff panel*/
-		TSharedPtr<FUICommandList> GraphEditorCommands;
-	};
+	/** Event handler that updates the graph view when user selects a new graph */
+	virtual void HandleGraphChanged( const FString& GraphName );
 
 	/*The two panels used to show the old & new revision*/ 
 	FDiffPanel				PanelOld, PanelNew;
