#ue4 build fix

diff --git a/Engine/Source/Developer/Merge/Private/SBlueprintMerge.cpp b/Engine/Source/Developer/Merge/Private/SBlueprintMerge.cpp
--- a/Engine/Source/Developer/Merge/Private/SBlueprintMerge.cpp
+++ b/Engine/Source/Developer/Merge/Private/SBlueprintMerge.cpp
@@ -6,7 +6,7 @@
 #include "GraphDiffControl.h"
 #include "ISourceControlModule.h"
 
-#define LOCTEXT_NAMESPACE "SBlueprintDif"
+#define LOCTEXT_NAMESPACE "SBlueprintMerge"
 
 // todo doc: where doe this utility function belong!
 static UPackage* CreateTempPackage(FString Name)
