#ue4 build fix

diff --git a/Engine/Source/Developer/Merge/Private/Merge.cpp b/Engine/Source/Developer/Merge/Private/Merge.cpp
--- a/Engine/Source/Developer/Merge/Private/Merge.cpp
+++ b/Engine/Source/Developer/Merge/Private/Merge.cpp
@@ -243,3 +243,4 @@ void FMerge::Merge( const UBlueprint& Object )
 	}
 }
 
+#undef LOCTEXT_NAMESPACE
