#ue4 more uniform initialization that needs to be removed.

diff --git a/Engine/Source/Developer/Merge/Private/Merge.cpp b/Engine/Source/Developer/Merge/Private/Merge.cpp
--- a/Engine/Source/Developer/Merge/Private/Merge.cpp
+++ b/Engine/Source/Developer/Merge/Private/Merge.cpp
@@ -64,7 +64,8 @@ static UObject* LoadRevision(const FString& AssetName, const ISourceControlRevis
 
 static FRevisionInfo GetRevisionInfo(ISourceControlRevision const& FromRevision)
 {
-	return{ FromRevision.GetRevisionNumber(), FromRevision.GetCheckInIdentifier(), FromRevision.GetDate() };
+	FRevisionInfo Ret = { FromRevision.GetRevisionNumber(), FromRevision.GetCheckInIdentifier(), FromRevision.GetDate() };
+	return Ret;
 }
 
 static UObject* LoadHeadRev(const FString& PackageName, const FString& AssetName, FRevisionInfo& OutRevInfo)
@@ -186,11 +187,7 @@ void FMerge::Merge( const UBlueprint& Object )
 
 	if (RemoteBlueprint && BaseBlueprint)
 	{
-		FMergeDisplayArgs DisplayArgs
-		{
-			CurrentRevInfo,
-			BaseRevInfo
-		};
+		FMergeDisplayArgs DisplayArgs = { CurrentRevInfo, BaseRevInfo };
 
 		FText WindowTitle = FText::Format(LOCTEXT("BlueprintMerge", "{0} - Blueprint Merge"), FText::FromString(AssetName));
 
