#ue4 first pass of moving mereg tool into blueprint editor, still has tab restoring issues.. disabled by EnableAssetMerge ini flag

diff --git a/Engine/Source/Developer/Merge/Private/Merge.cpp b/Engine/Source/Developer/Merge/Private/Merge.cpp
--- a/Engine/Source/Developer/Merge/Private/Merge.cpp
+++ b/Engine/Source/Developer/Merge/Private/Merge.cpp
@@ -109,7 +109,7 @@ class FMerge : public IMerge
 	/** IModuleInterface implementation */
 	virtual void StartupModule() override;
 	virtual void ShutdownModule() override;
-	virtual void Merge(const UBlueprint& Object) override;
+	virtual TSharedPtr<SWidget> GenerateMergeWidget(const UBlueprint& Object, FBlueprintEditor& Editor) override;
 };
 IMPLEMENT_MODULE( FMerge, Merge )
 
@@ -124,55 +124,8 @@ void FMerge::ShutdownModule()
 	// we call this function before unloading the module.
 }
 
-void FMerge::Merge( const UBlueprint& Object )
+TSharedPtr<SWidget> FMerge::GenerateMergeWidget(const UBlueprint& Object, FBlueprintEditor& Editor)
 {
-	// Verify that the object is not open in some other editor:
-	TSharedPtr< IToolkit > FoundAssetEditor = FToolkitManager::Get().FindEditorForAsset(&Object);
-	if (FoundAssetEditor.IsValid())
-	{
-		FText ToolkitName = FoundAssetEditor->GetBaseToolkitName();
-		FText ErrorMessage = FText::Format(LOCTEXT("CloseOpenEditors", "{0} must be closed before merging {1}. Would you like to close {0} and continue?"), ToolkitName, FText::FromString(Object.GetName()));
-		EAppReturnType::Type Result = FMessageDialog::Open(EAppMsgType::OkCancel, ErrorMessage);
-		if (Result == EAppReturnType::Cancel)
-		{
-			return;
-		}
-		else
-		{
-			FToolkitManager::Get().CloseToolkit(FoundAssetEditor.ToSharedRef());
-			FoundAssetEditor = FToolkitManager::Get().FindEditorForAsset(&Object);
-			if (FoundAssetEditor.IsValid())
-			{
-				ErrorMessage = FText::Format(LOCTEXT("MergeAborted", "Aborted Merge of {0} because {1} failed to close"), FText::FromString(Object.GetName()), ToolkitName);
-				FNotificationInfo Info(ErrorMessage);
-				Info.ExpireDuration = 5.0f;
-				FSlateNotificationManager::Get().AddNotification(Info);
-				return;
-			}
-		}
-	}
-
-	// verify that the object is not dirty:
-	if (Object.GetOutermost()->IsDirty())
-	{
-		FText RequestSave = FText::Format(LOCTEXT("RequestSaveForMerge", "{0} must be saved before it can be merged. Would you like to save and continue?"), FText::FromString(Object.GetName()));
-		EAppReturnType::Type Result = FMessageDialog::Open(EAppMsgType::OkCancel, RequestSave);
-		if (Result == EAppReturnType::Cancel)
-		{
-			return;
-		}
-		else
-		{
-			TArray<UPackage*> PackagesToSave;
-			PackagesToSave.Add(Object.GetOutermost());
-			const auto Result = FEditorFileUtils::PromptForCheckoutAndSave(PackagesToSave, /*bCheckDirty=*/ false, /*bPromptToSave=*/ false);
-			if (Result != FEditorFileUtils::PR_Success)
-			{
-				return;
-			}
-		}
-	}
-
 	// merge the local asset with the depot, SCC provides us with the last common revision as
 	// a basis for the merge:
 
@@ -189,33 +142,15 @@ void FMerge::Merge( const UBlueprint& Object )
 	{
 		FMergeDisplayArgs DisplayArgs = { CurrentRevInfo, BaseRevInfo };
 
-		FText WindowTitle = FText::Format(LOCTEXT("BlueprintMerge", "{0} - Blueprint Merge"), FText::FromString(AssetName));
-
-		const TSharedPtr<SWindow> Window = SNew(SWindow)
-			.Title(WindowTitle)
-			.ClientSize(FVector2D(1000, 800));
-
 		SBlueprintDiff::FArguments BaseArgs;
 		BaseArgs.BlueprintOld(BaseBlueprint)
 			.OldRevision(CurrentRevInfo)
 			.BlueprintNew(RemoteBlueprint)
 			.NewRevision(BaseRevInfo);
 
-		Window->SetContent(SNew(SBlueprintMerge)
+		return SNew(SBlueprintMerge)
 			.BlueprintLocal(const_cast<UBlueprint*>(&Object))
-			.OwningWindow(Window)
-			.BaseArgs(BaseArgs));
-
-		// Make this window a child of the modal window if we've been spawned while one is active.
-		TSharedPtr<SWindow> ActiveModal = FSlateApplication::Get().GetActiveModalWindow();
-		if (ActiveModal.IsValid())
-		{
-			FSlateApplication::Get().AddWindowAsNativeChild(Window.ToSharedRef(), ActiveModal.ToSharedRef());
-		}
-		else
-		{
-			FSlateApplication::Get().AddWindow(Window.ToSharedRef());
-		}
+			.BaseArgs(BaseArgs);
 	}
 	else
 	{
@@ -237,6 +172,7 @@ void FMerge::Merge( const UBlueprint& Object )
 		FNotificationInfo Info(ErrorMessage);
 		Info.ExpireDuration = 3.0f;
 		FSlateNotificationManager::Get().AddNotification(Info);
+		return TSharedPtr<SWidget>();
 	}
 }
 
diff --git a/Engine/Source/Developer/Merge/Private/SBlueprintMerge.cpp b/Engine/Source/Developer/Merge/Private/SBlueprintMerge.cpp
--- a/Engine/Source/Developer/Merge/Private/SBlueprintMerge.cpp
+++ b/Engine/Source/Developer/Merge/Private/SBlueprintMerge.cpp
@@ -5,6 +5,7 @@
 #include "DiffResults.h"
 #include "GraphDiffControl.h"
 #include "ISourceControlModule.h"
+#include "BlueprintEditor.h"
 
 #define LOCTEXT_NAMESPACE "SBlueprintMerge"
 
@@ -16,16 +17,6 @@ static UPackage* CreateTempPackage(FString Name)
 	return CreatePackage(NULL, *TempPackageName);
 }
 
-static UBlueprint* DuplicateBlueprint( const UBlueprint* const BlueprintToClone)
-{
-	UPackage* TempPackage = CreateTempPackage(BlueprintToClone->GetName());
-	TArray<UObject*> Subobjects;
-	GetObjectsWithOuter(TempPackage, Subobjects, true);
-	check(Subobjects.Num() == 0);
-	EObjectFlags FlagMask = RF_AllFlags & ~RF_Transient & ~RF_PendingKill;
-	return Cast<UBlueprint>(StaticDuplicateObject(BlueprintToClone, TempPackage, *BlueprintToClone->GetName(), FlagMask));
-}
-
 template< class Predicate >
 UObject* FindOuter( UObject& Obj, Predicate P )
 {
@@ -120,15 +111,19 @@ public:
 	}
 };
 
+SBlueprintMerge::SBlueprintMerge()
+	: SBlueprintDiff()
+	, OwningEditor()
+{
+}
+
 void SBlueprintMerge::Construct(const FArguments& InArgs)
 {
 	FMergeDifferencesListCommands::Register();
 
 	PanelLocal.Blueprint = InArgs._BlueprintLocal;
-	BlueprintResult = DuplicateBlueprint(InArgs._BlueprintLocal);
-	(void)PackagesAppearEquivalent( BlueprintResult->GetOutermost(), PanelLocal.Blueprint->GetOutermost() );
-	//check( PackagesAppearEquivalent( BlueprintResult->GetOutermost(), PanelLocal.Blueprint->GetOutermost() ) );
-	OwningWindow = InArgs._OwningWindow;
+
+	OwningEditor = InArgs._OwningEditor;
 
 	return SBlueprintDiff::Construct( InArgs._BaseArgs );
 }
@@ -275,11 +270,6 @@ TSharedRef<SWidget> SBlueprintMerge::GenerateDiffWindow()
 					]
 				]
 			]
-			+ SSplitter::Slot()
-			.Value(.3f)
-			[
-				SAssignNew(EditorBorder, SBorder)
-			]
 		]
 	;
 }
@@ -364,7 +354,6 @@ TSharedRef<SWidget> SBlueprintMerge::GenerateToolbar()
 		];
 }
 
-// @cr doc const shallowness
 static UEdGraph* FindGraphByName(UBlueprint const& FromBlueprint, const FString& GraphName)
 {
 	TArray<UEdGraph*> Graphs;
@@ -383,45 +372,10 @@ void SBlueprintMerge::HandleGraphChanged(const FString& GraphName)
 	UEdGraph* GraphOld = FindGraphByName(*PanelOld.Blueprint, GraphName);
 	UEdGraph* GraphNew = FindGraphByName(*PanelNew.Blueprint, GraphName);
 	UEdGraph* GraphLocal = FindGraphByName(*PanelLocal.Blueprint, GraphName);
-	UEdGraph* GraphResult = FindGraphByName(*BlueprintResult, GraphName);
 
 	PanelOld.GeneratePanel(GraphOld, NULL);
 	PanelNew.GeneratePanel(GraphNew, GraphOld);
 	PanelLocal.GeneratePanel(GraphLocal, GraphOld);
-
-	{
-		// Set up a normal editor window so that the user can tweak the result of the merge:
-		// @todo doc: need to copy the target blueprint and destructively edit that one, not 
-		// the base blueprint..
-		if (GraphResult)
-		{
-			SGraphEditor::FGraphEditorEvents InEvents;
-
-			//FGraphAppearanceInfo AppearanceInfo;
-			//AppearanceInfo.CornerText = LOCTEXT("AppearanceCornerText_BlueprintDif", "DIFF").ToString();
-
-			/*if (!GraphEditorCommands.IsValid())
-			{
-				GraphEditorCommands = MakeShareable(new FUICommandList());
-
-				GraphEditorCommands->MapAction(FGenericCommands::Get().Copy,
-					FExecuteAction::CreateRaw(this, &FDiffPanel::CopySelectedNodes),
-					FCanExecuteAction::CreateRaw(this, &FDiffPanel::CanCopyNodes)
-					);
-			}*/
-
-			auto Editor = SNew(SGraphEditor)
-				.GraphToEdit(GraphResult)
-				.TitleBarEnabledOnly(false)
-				//.Appearance(AppearanceInfo)
-				.GraphEvents(InEvents);
-			EditorBorder->SetContent(Editor);
-		}
-		else
-		{
-			EditorBorder->ClearContent();
-		}
-	}
 }
 
 template< typename T >
@@ -577,12 +531,10 @@ FReply SBlueprintMerge::OnAcceptResultClicked()
 
 	// cr doc: this is too simple to be correct..
 	PanelLocal.Blueprint->RemoveGeneratedClasses();
-	StaticDuplicateObject( BlueprintResult, Package, *PanelLocal.Blueprint->GetName() );
 
-	auto OwningWindowPinned = OwningWindow.Pin();
-	if (OwningWindowPinned.IsValid())
+	if (OwningEditor.IsValid())
 	{
-		OwningWindowPinned->RequestDestroyWindow();
+		OwningEditor->CloseWindow();
 	}
 
 	const auto Result = FEditorFileUtils::PromptForCheckoutAndSave(PackagesToSave, /*bCheckDirty=*/ false, /*bPromptToSave=*/ false);
@@ -600,10 +552,9 @@ FReply SBlueprintMerge::OnAcceptResultClicked()
 
 FReply SBlueprintMerge::OnCancelClicked()
 {
-	auto OwningWindowPinned = OwningWindow.Pin();
-	if( OwningWindowPinned.IsValid() )
+	if (OwningEditor.IsValid())
 	{
-		OwningWindowPinned->RequestDestroyWindow();
+		OwningEditor->CloseWindow();
 	}
 	return FReply::Handled();
 }
@@ -628,42 +579,18 @@ FReply SBlueprintMerge::OnTakeBaseClicked()
 
 void SBlueprintMerge::StageBlueprint(UBlueprint const* DesiredBP)
 {
-	// User has decided to discard the remote changes, update the result window
-	// with this result:
-	if (BlueprintResult)
-	{
-		//delete BlueprintResult;... yeah.. i have no idea what's keeping this stuff alive, best not to delete it
-	}
+	// Take the DesiredBP and applies its changes to the result blueprint:
+	UBlueprint* BlueprintResult = OwningEditor->GetBlueprintObj();
+
+	// I want everything in the package... components, CDO, etc:
 
-	BlueprintResult = DuplicateBlueprint(DesiredBP);
 
 	(void)PackagesAppearEquivalent(BlueprintResult->GetOutermost(), DesiredBP->GetOutermost());
-	//check( PackagesAppearEquivalent( BlueprintResult->GetOutermost(), DesiredBP->GetOutermost() ) );
+	check( PackagesAppearEquivalent( BlueprintResult->GetOutermost(), DesiredBP->GetOutermost() ) );
 
 	SGraphEditor::FGraphEditorEvents InEvents;
 	TArray< FGraphToDiff>  Selected = GraphsToDiff->GetSelectedItems();
 	check(Selected.Num() <= 1); // We set this control up to not support multiselect
-
-	UEdGraph* GraphResult = NULL;
-
-	if (Selected.Num() == 1)
-	{
-		GraphResult = FindGraphByName(*BlueprintResult, Selected[0]->GetGraphNew()->GetName());
-	}
-
-	if( GraphResult )
-	{
-		auto Editor = SNew(SGraphEditor)
-			.GraphToEdit(GraphResult)
-			.TitleBarEnabledOnly(false)
-			//.Appearance(AppearanceInfo)
-			.GraphEvents(InEvents);
-		EditorBorder->SetContent(Editor);
-	}
-	else
-	{
-		EditorBorder->ClearContent();
-	}
 }
 
 #undef LOCTEXT_NAMESPACE
diff --git a/Engine/Source/Developer/Merge/Private/SBlueprintMerge.h b/Engine/Source/Developer/Merge/Private/SBlueprintMerge.h
--- a/Engine/Source/Developer/Merge/Private/SBlueprintMerge.h
+++ b/Engine/Source/Developer/Merge/Private/SBlueprintMerge.h
@@ -3,13 +3,17 @@
 
 #include "SBlueprintDiff.h"
 
+class FBlueprintEditor;
+
 class MERGE_API SBlueprintMerge : public SBlueprintDiff
 {
 public:
+	SBlueprintMerge();
+
 	SLATE_BEGIN_ARGS(SBlueprintMerge){}
 	SLATE_ARGUMENT(class UBlueprint*, BlueprintLocal)
 	SLATE_ARGUMENT( SBlueprintDiff::FArguments, BaseArgs )
-	SLATE_ARGUMENT( TWeakPtr<SWindow>, OwningWindow )
+	SLATE_ARGUMENT( TSharedPtr<FBlueprintEditor>, OwningEditor )
 	SLATE_END_ARGS()
 	
 	void Construct(const FArguments& InArgs);
@@ -17,7 +21,7 @@ public:
 	void OnDiffListSelectionChanged(TSharedPtr<struct FDiffSingleResult> Item, ESelectInfo::Type SelectionType);
 
 protected:
-	/** Overrides, see base class for documentation */
+	/** SBlueprintDiff overrides, see base class for documentation */
 	virtual void ResetGraphEditors() override;
 	virtual TSharedRef<SWidget> GenerateDiffWindow() override;
 	virtual TSharedRef<SWidget> GenerateToolbar() override;
@@ -36,10 +40,8 @@ protected:
 	FReply OnTakeBaseClicked();
 	void StageBlueprint(UBlueprint const* DesiredBP);
 
-	// Raw pointer?!
-	UBlueprint* BlueprintResult;
+	TSharedPtr<FBlueprintEditor> OwningEditor;
 	FDiffPanel PanelLocal;
-	TSharedPtr<SBorder>		 EditorBorder;
 	TWeakPtr<SWindow>		OwningWindow;
 
 	// This has to be allocated here because SListView cannot own the list
diff --git a/Engine/Source/Developer/Merge/Public/Merge.h b/Engine/Source/Developer/Merge/Public/Merge.h
--- a/Engine/Source/Developer/Merge/Public/Merge.h
+++ b/Engine/Source/Developer/Merge/Public/Merge.h
@@ -6,6 +6,7 @@
 #include "IAssetTypeActions.h"
 
 class UBlueprint;
+class FBlueprintEditor;
 
 struct FMergeDisplayArgs
 {
@@ -43,11 +44,10 @@ public:
 	}
 
 	/**
-	 * Performs a three way merge of AssetA and AssetB using CommonOrigin as a base. Takes a display args structure
-	 * to customize the view.
+	 * Generates a widget used to perform a three way merge operation.
 	 *
-	 * @return The resulting UBlueprint, NULL if merge was canceled or otherwise failed
+	 * @return The merge view widget
 	 */
-	virtual void Merge( const UBlueprint& Object ) = 0;
+	virtual TSharedPtr<SWidget> GenerateMergeWidget( const UBlueprint& Object, FBlueprintEditor& Editor ) = 0;
 };
 
diff --git a/Engine/Source/Editor/ContentBrowser/ContentBrowser.Build.cs b/Engine/Source/Editor/ContentBrowser/ContentBrowser.Build.cs
--- a/Engine/Source/Editor/ContentBrowser/ContentBrowser.Build.cs
+++ b/Engine/Source/Editor/ContentBrowser/ContentBrowser.Build.cs
@@ -24,7 +24,6 @@ public class ContentBrowser : ModuleRules
 				"CoreUObject",
                 "InputCore",
 				"Engine",
-                "Merge",
 				"Slate",
 				"SlateCore",
                 "EditorStyle",
diff --git a/Engine/Source/Editor/ContentBrowser/Private/AssetContextMenu.cpp b/Engine/Source/Editor/ContentBrowser/Private/AssetContextMenu.cpp
--- a/Engine/Source/Editor/ContentBrowser/Private/AssetContextMenu.cpp
+++ b/Engine/Source/Editor/ContentBrowser/Private/AssetContextMenu.cpp
@@ -22,7 +22,6 @@
 #include "KismetEditorUtilities.h"
 #include "AssetToolsModule.h"
 #include "ComponentAssetBroker.h"
-#include "Merge.h"
 
 #define LOCTEXT_NAMESPACE "ContentBrowser"
 
@@ -33,7 +32,6 @@ FAssetContextMenu::FAssetContextMenu(const TWeakPtr<SAssetView>& InAssetView)
 	, bCanExecuteSCCCheckIn(false)
 	, bCanExecuteSCCHistory(false)
 	, bCanExecuteSCCRevert(false)
-	, bCanExecuteSCCMerge(false)
 	, bCanExecuteSCCSync(false)
 	, AssetView(InAssetView)
 {
@@ -407,19 +405,6 @@ bool FAssetContextMenu::AddSourceControlMenuOptions(FMenuBuilder& MenuBuilder)
 				);
 		}
 
-		if (CanExecuteSCCMerge())
-		{
-			MenuBuilder.AddMenuEntry(
-				LOCTEXT("SCCMerge", "Merge"),
-				LOCTEXT("SCCMergeTooltip", "Initiates a three-way merge of an asset."),
-				FSlateIcon(),
-				FUIAction(
-				FExecuteAction::CreateSP(this, &FAssetContextMenu::ExecuteSCCMerge),
-				FCanExecuteAction::CreateSP(this, &FAssetContextMenu::CanExecuteSCCMerge)
-				)
-				);
-		}
-
 		MenuBuilder.AddMenuEntry(
 			LOCTEXT("SCCRefresh", "Refresh"),
 			LOCTEXT("SCCRefreshTooltip", "Updates the source control status of the asset."),
@@ -1122,24 +1107,6 @@ void FAssetContextMenu::ExecuteSCCRevert()
 	FSourceControlWindows::PromptForRevert(PackageNames);
 }
 
-void FAssetContextMenu::ExecuteSCCMerge()
-{
-	TArray<FString> PackageNames;
-	GetSelectedPackageNames(PackageNames);
-
-	for (int32 AssetIdx = 0; AssetIdx < SelectedAssets.Num(); ++AssetIdx)
-	{
-		const FAssetData& Asset = SelectedAssets[AssetIdx];
-		UObject* Object = Asset.GetAsset();
-
-		if( Object )
-		{
-			IMerge::Get().Merge(*Cast<UBlueprint>(Object));
-		}
-	}
-
-}
-
 void FAssetContextMenu::ExecuteSCCSync()
 {
 	TArray<FString> PackageNames;
@@ -1270,13 +1237,6 @@ bool FAssetContextMenu::CanExecuteSCCRevert() const
 	return bCanExecuteSCCRevert;
 }
 
-bool FAssetContextMenu::CanExecuteSCCMerge() const
-{
-	bool bEnableMerge = false;
-	GConfig->GetBool(TEXT("AssetMerge"), TEXT("EnableAssetMerge"), bEnableMerge, GEngineIni);
-	return bCanExecuteSCCMerge && bEnableMerge;
-}
-
 bool FAssetContextMenu::CanExecuteSCCSync() const
 {
 	return bCanExecuteSCCSync;
@@ -1372,7 +1332,6 @@ void FAssetContextMenu::CacheCanExecuteVars()
 	bCanExecuteSCCCheckIn = false;
 	bCanExecuteSCCHistory = false;
 	bCanExecuteSCCRevert = false;
-	bCanExecuteSCCMerge = false;
 	bCanExecuteSCCSync = false;
 
 	for (auto AssetIt = SelectedAssets.CreateConstIterator(); AssetIt; ++AssetIt)
@@ -1419,11 +1378,6 @@ void FAssetContextMenu::CacheCanExecuteVars()
 					bCanExecuteSCCCheckIn = true;
 					bCanExecuteSCCRevert = true;
 				}
-
-				if( SourceControlState->IsConflicted() && SourceControlState->IsCurrent() )
-				{
-					bCanExecuteSCCMerge = true;
-				}
 			}
 		}
 
diff --git a/Engine/Source/Editor/ContentBrowser/Private/AssetContextMenu.h b/Engine/Source/Editor/ContentBrowser/Private/AssetContextMenu.h
--- a/Engine/Source/Editor/ContentBrowser/Private/AssetContextMenu.h
+++ b/Engine/Source/Editor/ContentBrowser/Private/AssetContextMenu.h
@@ -133,9 +133,6 @@ private:
 	/** Handler for when "Source Control Revert" is selected */
 	void ExecuteSCCRevert();
 
-	/** Handler for when "Source Control Merge" is selected */
-	void ExecuteSCCMerge();
-
 	/** Handler for when "Source Control Sync" is selected */
 	void ExecuteSCCSync();
 
@@ -190,9 +187,6 @@ private:
 	/** Handler to check to see if "Source Control Revert" can be executed */
 	bool CanExecuteSCCRevert() const;
 
-	/** Handler to check to see if "Source Control Merge" can be executed */
-	bool CanExecuteSCCMerge() const;
-
 	/** Handler to check to see if "Source Control Sync" can be executed */
 	bool CanExecuteSCCSync() const;
 
@@ -252,6 +246,5 @@ private:
 	bool bCanExecuteSCCCheckIn;
 	bool bCanExecuteSCCHistory;
 	bool bCanExecuteSCCRevert;
-	bool bCanExecuteSCCMerge;
 	bool bCanExecuteSCCSync;
 };
diff --git a/Engine/Source/Editor/Kismet/Kismet.Build.cs b/Engine/Source/Editor/Kismet/Kismet.Build.cs
--- a/Engine/Source/Editor/Kismet/Kismet.Build.cs
+++ b/Engine/Source/Editor/Kismet/Kismet.Build.cs
@@ -26,6 +26,7 @@ public class Kismet : ModuleRules
 				"SlateCore",
                 "EditorStyle",
 				"Engine",
+				"Merge",
 				"MessageLog",
 				"UnrealEd",
 				"GraphEditor",
diff --git a/Engine/Source/Editor/Kismet/Private/BlueprintEditor.cpp b/Engine/Source/Editor/Kismet/Private/BlueprintEditor.cpp
--- a/Engine/Source/Editor/Kismet/Private/BlueprintEditor.cpp
+++ b/Engine/Source/Editor/Kismet/Private/BlueprintEditor.cpp
@@ -74,7 +74,8 @@
 #include "ISCSEditorCustomization.h"
 #include "Editor/UnrealEd/Public/SourceCodeNavigation.h"
 
-
+// Blueprint merging
+#include "Merge.h"
 
 #define LOCTEXT_NAMESPACE "BlueprintEditor"
 
@@ -1591,7 +1592,8 @@ void FBlueprintEditor::CreateDefaultTabContents(const TArray<UBlueprint*>& InBlu
 		];
 
 	FindResults = SNew(SFindInBlueprints, SharedThis(this));
-
+	MergeTool = SNew(SHorizontalBox);
+	
 	this->Inspector = 
 		SNew(SKismetInspector)
 		. HideNameArea(true)
@@ -1846,6 +1848,10 @@ void FBlueprintEditor::CreateDefaultCommands()
 		FIsActionButtonVisible::CreateSP(this, &FBlueprintEditor::NewDocument_IsVisibleForType, CGT_NewAnimationGraph));
 	*/
 
+	ToolkitCommands->MapAction(FBlueprintEditorCommands::Get().SaveIntermediateBuildProducts,
+		FExecuteAction::CreateSP(this, &FBlueprintEditor::ToggleSaveIntermediateBuildProducts),
+		FCanExecuteAction(),
+		FIsActionChecked::CreateSP(this, &FBlueprintEditor::GetSaveIntermediateBuildProducts));
 	ToolkitCommands->MapAction( FBlueprintEditorCommands::Get().RecompileGraphEditor,
 		FExecuteAction::CreateStatic( &FLocalKismetCallbacks::RecompileGraphEditor_OnClicked ),
 		FCanExecuteAction::CreateStatic( &FLocalKismetCallbacks::CanRecompileModules ));
@@ -1859,10 +1865,9 @@ void FBlueprintEditor::CreateDefaultCommands()
 		FExecuteAction::CreateStatic( &FLocalKismetCallbacks::RecompilePersona_OnClicked ),
 		FCanExecuteAction::CreateStatic( &FLocalKismetCallbacks::CanRecompileModules ));
 
-	ToolkitCommands->MapAction( FBlueprintEditorCommands::Get().SaveIntermediateBuildProducts,
-		FExecuteAction::CreateSP(this, &FBlueprintEditor::ToggleSaveIntermediateBuildProducts),
-		FCanExecuteAction(),
-		FIsActionChecked::CreateSP(this, &FBlueprintEditor::GetSaveIntermediateBuildProducts) );
+	ToolkitCommands->MapAction(FBlueprintEditorCommands::Get().BeginBlueprintMerge,
+		FExecuteAction::CreateSP(this, &FBlueprintEditor::OnBeginBlueprintMerge),
+		FCanExecuteAction());
 }
 
 void FBlueprintEditor::FindInBlueprint_Clicked()
@@ -5447,6 +5452,16 @@ void FBlueprintEditor::OnRepairCorruptedBlueprint()
 	Compiler.RecoverCorruptedBlueprint(GetBlueprintObj());
 }
 
+void FBlueprintEditor::OnBeginBlueprintMerge()
+{
+	// @cr doc: any way to show this in all of the subviews (defaults + components + graph view, rather than just the active one?)
+	//if( MergeTool->Slots().Num() == 0 && GetBlueprintObj() )
+	{
+		MergeTool->AddSlot()[ IMerge::Get().GenerateMergeWidget(*GetBlueprintObj(), *this).ToSharedRef() ];
+	}
+	TabManager->InvokeTab(FBlueprintEditorTabs::MergeToolID);
+}
+
 void FBlueprintEditor::StartEditingDefaults(bool bAutoFocus, bool bForceRefresh)
 {
 	if (IsEditingSingleBlueprint())
diff --git a/Engine/Source/Editor/Kismet/Private/BlueprintEditorCommands.cpp b/Engine/Source/Editor/Kismet/Private/BlueprintEditorCommands.cpp
--- a/Engine/Source/Editor/Kismet/Private/BlueprintEditorCommands.cpp
+++ b/Engine/Source/Editor/Kismet/Private/BlueprintEditorCommands.cpp
@@ -61,6 +61,9 @@ void FBlueprintEditorCommands::RegisterCommands()
 	UI_COMMAND( RecompileKismetCompiler, "Recompile Blueprint Compiler", "Recompiles and reloads C++ code for the blueprint compiler", EUserInterfaceActionType::Button, FInputGesture() );
 	UI_COMMAND( RecompileBlueprintEditor, "Recompile Blueprint Editor", "Recompiles and reloads C++ code for the blueprint editor", EUserInterfaceActionType::Button, FInputGesture() );
 	UI_COMMAND( RecompilePersona, "Recompile Persona", "Recompiles and reloads C++ code for Persona", EUserInterfaceActionType::Button, FInputGesture() );
+
+	// SCC commands
+	UI_COMMAND( BeginBlueprintMerge, "Merge", "Shows the Blueprint merge panel and toolbar, allowing the user to resolve conflicted blueprints", EUserInterfaceActionType::Button, FInputGesture() );
 }
 
 PRAGMA_ENABLE_OPTIMIZATION
diff --git a/Engine/Source/Editor/Kismet/Private/BlueprintEditorCommands.h b/Engine/Source/Editor/Kismet/Private/BlueprintEditorCommands.h
--- a/Engine/Source/Editor/Kismet/Private/BlueprintEditorCommands.h
+++ b/Engine/Source/Editor/Kismet/Private/BlueprintEditorCommands.h
@@ -66,6 +66,9 @@ public:
 	TSharedPtr< FUICommandInfo > RecompileKismetCompiler;
 	TSharedPtr< FUICommandInfo > RecompileBlueprintEditor;
 	TSharedPtr< FUICommandInfo > RecompilePersona;
+
+	// SSC commands
+	TSharedPtr< FUICommandInfo > BeginBlueprintMerge;
 };
 
 //////////////////////////////////////////////////////////////////////////
diff --git a/Engine/Source/Editor/Kismet/Private/BlueprintEditorModes.cpp b/Engine/Source/Editor/Kismet/Private/BlueprintEditorModes.cpp
--- a/Engine/Source/Editor/Kismet/Private/BlueprintEditorModes.cpp
+++ b/Engine/Source/Editor/Kismet/Private/BlueprintEditorModes.cpp
@@ -52,6 +52,7 @@ FBlueprintEditorApplicationMode::FBlueprintEditorApplicationMode(TSharedPtr<clas
 	BlueprintEditorTabFactories.RegisterFactory(MakeShareable(new FMyBlueprintSummoner(InBlueprintEditor)));
 	BlueprintEditorTabFactories.RegisterFactory(MakeShareable(new FCompilerResultsSummoner(InBlueprintEditor)));
 	BlueprintEditorTabFactories.RegisterFactory(MakeShareable(new FFindResultsSummoner(InBlueprintEditor)));
+	BlueprintEditorTabFactories.RegisterFactory(MakeShareable(new FMergeToolSummoner(InBlueprintEditor)));
 	if( bRegisterViewport )
 	{
 		BlueprintEditorTabFactories.RegisterFactory(MakeShareable(new FSCSViewportSummoner(InBlueprintEditor)));
@@ -120,6 +121,16 @@ FBlueprintEditorApplicationMode::FBlueprintEditorApplicationMode(TSharedPtr<clas
 					)
 				)
 			)
+			->Split
+			(
+				FTabManager::NewSplitter() ->SetOrientation( Orient_Horizontal )
+				->SetSizeCoefficient(0.35f)
+				->Split
+				(
+					FTabManager::NewStack()
+					->AddTab( FBlueprintEditorTabs::MergeToolID, ETabState::ClosedTab )
+				)
+			)
 		);
 	
 	// setup toolbar
@@ -173,7 +184,8 @@ FBlueprintDefaultsApplicationMode::FBlueprintDefaultsApplicationMode(TSharedPtr<
 	
 	BlueprintDefaultsTabFactories.RegisterFactory(MakeShareable(new FDefaultsEditorSummoner(InBlueprintEditor)));
 	BlueprintDefaultsTabFactories.RegisterFactory(MakeShareable(new FFindResultsSummoner(InBlueprintEditor)));
-	
+	BlueprintDefaultsTabFactories.RegisterFactory(MakeShareable(new FMergeToolSummoner(InBlueprintEditor)));
+
 	TabLayout = FTabManager::NewLayout( "Standalone_BlueprintDefaults_Layout_v4" )
 		->AddArea
 		(
@@ -192,6 +204,16 @@ FBlueprintDefaultsApplicationMode::FBlueprintDefaultsApplicationMode(TSharedPtr<
 				->SetHideTabWell(true)
 				->AddTab( FBlueprintEditorTabs::DefaultEditorID, ETabState::OpenedTab )
 			)
+			->Split
+			(
+				FTabManager::NewSplitter() ->SetOrientation( Orient_Horizontal )
+				->SetSizeCoefficient(0.35f)
+				->Split
+				(
+					FTabManager::NewStack()
+					->AddTab( FBlueprintEditorTabs::MergeToolID, ETabState::ClosedTab )
+				)
+			)
 		);
 
 	// setup toolbar
@@ -231,7 +253,8 @@ FBlueprintComponentsApplicationMode::FBlueprintComponentsApplicationMode(TShared
 	BlueprintComponentsTabFactories.RegisterFactory(MakeShareable(new FSelectionDetailsSummoner(InBlueprintEditor)));
 	BlueprintComponentsTabFactories.RegisterFactory(MakeShareable(new FDefaultsEditorSummoner(InBlueprintEditor)));
 	BlueprintComponentsTabFactories.RegisterFactory(MakeShareable(new FFindResultsSummoner(InBlueprintEditor)));
-	
+	BlueprintComponentsTabFactories.RegisterFactory(MakeShareable(new FMergeToolSummoner(InBlueprintEditor)));
+
 	TabLayout = FTabManager::NewLayout( "Standalone_BlueprintComponents_Layout_v5" )
 		->AddArea
 		(
@@ -274,6 +297,16 @@ FBlueprintComponentsApplicationMode::FBlueprintComponentsApplicationMode(TShared
 					->AddTab( FBlueprintEditorTabs::SCSViewportID, ETabState::OpenedTab )
 				)
 			)
+			->Split
+			(
+				FTabManager::NewSplitter() ->SetOrientation( Orient_Horizontal )
+				->SetSizeCoefficient(0.35f)
+				->Split
+				(
+					FTabManager::NewStack()
+					->AddTab( FBlueprintEditorTabs::MergeToolID, ETabState::ClosedTab )
+				)
+			)
 		);
 
 	// setup toolbar
diff --git a/Engine/Source/Editor/Kismet/Private/BlueprintEditorTabFactories.h b/Engine/Source/Editor/Kismet/Private/BlueprintEditorTabFactories.h
--- a/Engine/Source/Editor/Kismet/Private/BlueprintEditorTabFactories.h
+++ b/Engine/Source/Editor/Kismet/Private/BlueprintEditorTabFactories.h
@@ -488,4 +488,31 @@ public:
 	}
 };
 
+//////////////////////////////////////////////////////////////////////////
+// FMergeToolSummoner
+
+struct FMergeToolSummoner : public FWorkflowTabFactory
+{
+	FMergeToolSummoner(TSharedPtr<class FAssetEditorToolkit> InHostingApp)
+		: FWorkflowTabFactory(FBlueprintEditorTabs::MergeToolID, InHostingApp)
+	{
+		TabLabel = LOCTEXT("MergeToolTabTitle", "Merge Tool");
+		bIsSingleton = true;
+
+		ViewMenuDescription = LOCTEXT("MergeToolView", "Merge Tool");
+		ViewMenuTooltip = LOCTEXT("MergeToolView_ToolTip", "Displays three revisions of a blueprint for merging purposes, provides access to SSC merge operations");
+	}
+
+	virtual TSharedRef<SWidget> CreateTabBody(const FWorkflowTabSpawnInfo& Info) const override
+	{
+		TSharedPtr<FBlueprintEditor> BlueprintEditorPtr = StaticCastSharedPtr<FBlueprintEditor>(HostingApp.Pin());
+		return BlueprintEditorPtr->GetMergeTool();
+	}
+
+	virtual FText GetTabToolTipText(const FWorkflowTabSpawnInfo& Info) const override
+	{
+		return LOCTEXT("MergeToolTooltip", "The merge tool allows users to merge two blueprints that diverged from a common base blueprint");
+	}
+};
+
 #undef LOCTEXT_NAMESPACE
diff --git a/Engine/Source/Editor/Kismet/Private/BlueprintEditorTabs.cpp b/Engine/Source/Editor/Kismet/Private/BlueprintEditorTabs.cpp
--- a/Engine/Source/Editor/Kismet/Private/BlueprintEditorTabs.cpp
+++ b/Engine/Source/Editor/Kismet/Private/BlueprintEditorTabs.cpp
@@ -9,6 +9,7 @@ const FName FBlueprintEditorTabs::DebugID( TEXT( "Debug" ) );
 const FName FBlueprintEditorTabs::PaletteID( TEXT( "PaletteList" ) );
 const FName FBlueprintEditorTabs::CompilerResultsID( TEXT( "CompilerResults" ) );
 const FName FBlueprintEditorTabs::FindResultsID( TEXT( "FindResults" ) );
+const FName FBlueprintEditorTabs::MergeToolID(TEXT("MergeTool"));
 const FName FBlueprintEditorTabs::ConstructionScriptEditorID( TEXT( "ConstructionScriptEditor" ) );
 const FName FBlueprintEditorTabs::SCSViewportID( TEXT( "SCSViewport" ) );
 const FName FBlueprintEditorTabs::MyBlueprintID( TEXT( "MyBlueprint" ) );
diff --git a/Engine/Source/Editor/Kismet/Private/SBlueprintEditorToolbar.cpp b/Engine/Source/Editor/Kismet/Private/SBlueprintEditorToolbar.cpp
--- a/Engine/Source/Editor/Kismet/Private/SBlueprintEditorToolbar.cpp
+++ b/Engine/Source/Editor/Kismet/Private/SBlueprintEditorToolbar.cpp
@@ -53,6 +53,14 @@ public:
 	// End of SWidget interface
 };
 
+static bool PendingMerge( const UBlueprint& BlueprintObj )
+{
+	bool bEnableMerge = false;
+	GConfig->GetBool(TEXT("AssetMerge"), TEXT("EnableAssetMerge"), bEnableMerge, GEngineIni);
+	ISourceControlProvider& SourceControlProvider = ISourceControlModule::Get().GetProvider();
+	FSourceControlStatePtr SourceControlState = SourceControlProvider.GetState(BlueprintObj.GetOutermost(), EStateCacheUsage::Use);
+	return bEnableMerge && SourceControlState.IsValid() && SourceControlState->IsConflicted();
+}
 
 //////////////////////////////////////////////////////////////////////////
 // FKismet2Menu
@@ -69,6 +77,12 @@ void FKismet2Menu::FillFileMenuBlueprintSection( FMenuBuilder& MenuBuilder, FBlu
 			LOCTEXT("BlueprintEditorDiffToolTip", "Diff against previous revisions"),
 			FOnGetContent::CreateStatic< FBlueprintEditor& >( &FKismet2Menu::MakeDiffMenu, Kismet),
 			FSlateIcon());
+
+		UBlueprint* BlueprintObj = Kismet.GetBlueprintObj();
+		if (BlueprintObj && PendingMerge( *BlueprintObj ) )
+		{
+			MenuBuilder.AddMenuEntry(FBlueprintEditorCommands::Get().BeginBlueprintMerge);
+		}
 	}
 	MenuBuilder.EndSection();
 
diff --git a/Engine/Source/Editor/Kismet/Public/BlueprintEditor.h b/Engine/Source/Editor/Kismet/Public/BlueprintEditor.h
--- a/Engine/Source/Editor/Kismet/Public/BlueprintEditor.h
+++ b/Engine/Source/Editor/Kismet/Public/BlueprintEditor.h
@@ -191,6 +191,7 @@ public:
 	TSharedRef<class SBlueprintPalette> GetPalette() const {return Palette.ToSharedRef();}
 	TSharedRef<class SWidget> GetCompilerResults() const {return CompilerResults.ToSharedRef();}
 	TSharedRef<class SFindInBlueprints> GetFindResults() const {return FindResults.ToSharedRef();}
+	TSharedRef<class SWidget> GetMergeTool() const { return MergeTool.ToSharedRef(); }
 	TSharedRef<class SSCSEditor> GetSCSEditor() const {return SCSEditor.ToSharedRef();}
 	TSharedPtr<class SSCSEditorViewport> GetSCSViewport() const {return SCSViewport;}
 	TSharedPtr<class SMyBlueprint> GetMyBlueprintWidget() const {return MyBlueprintWidget;}
@@ -704,6 +705,8 @@ protected:
 	void OnListObjectsReferencedByBlueprint();
 	void OnRepairCorruptedBlueprint();
 
+	void OnBeginBlueprintMerge();
+
 	void OnNodeDoubleClicked(class UEdGraphNode* Node);
 
 	virtual void OnEditTabClosed(TSharedRef<SDockTab> Tab);
@@ -912,6 +915,9 @@ protected:
 	/** Find results log as well as the search filter */
 	TSharedPtr<class SFindInBlueprints> FindResults;
 
+	/** Merge tool */
+	TSharedPtr<class SHorizontalBox> MergeTool;
+
 	/** Reference to owner of the current popup */
 	TWeakPtr<class SWindow> NameEntryPopupWindow;
 
diff --git a/Engine/Source/Editor/Kismet/Public/BlueprintEditorTabs.h b/Engine/Source/Editor/Kismet/Public/BlueprintEditorTabs.h
--- a/Engine/Source/Editor/Kismet/Public/BlueprintEditorTabs.h
+++ b/Engine/Source/Editor/Kismet/Public/BlueprintEditorTabs.h
@@ -11,6 +11,7 @@ struct KISMET_API FBlueprintEditorTabs
 	static const FName PaletteID;
 	static const FName CompilerResultsID;
 	static const FName FindResultsID;
+	static const FName MergeToolID;
 	static const FName ConstructionScriptEditorID;
 	static const FName SCSViewportID;
 	static const FName MyBlueprintID;
