#ue4 build fix

diff --git a/Engine/Source/Developer/Merge/Public/Merge.h b/Engine/Source/Developer/Merge/Public/Merge.h
--- a/Engine/Source/Developer/Merge/Public/Merge.h
+++ b/Engine/Source/Developer/Merge/Public/Merge.h
@@ -48,6 +48,6 @@ public:
 	 *
 	 * @return The merge view widget
 	 */
-	virtual TSharedPtr<SWidget> GenerateMergeWidget( const UBlueprint& Object, FBlueprintEditor& Editor ) = 0;
+	virtual TSharedPtr<class SWidget> GenerateMergeWidget( const UBlueprint& Object, FBlueprintEditor& Editor ) = 0;
 };
 
