#UE4 Removed some ill advised buttons (take local, take remote, auto merge are difficult to implement for all blueprints), tweaked ownership (bluepint merge tool can keep shared ptr to blueprint editor - verified no shared ptr cycle when closing blueprint editor), corrected merge-accepted logic for embedded blueprint merge widget

diff --git a/Engine/Source/Developer/Merge/Private/SBlueprintMerge.cpp b/Engine/Source/Developer/Merge/Private/SBlueprintMerge.cpp
--- a/Engine/Source/Developer/Merge/Private/SBlueprintMerge.cpp
+++ b/Engine/Source/Developer/Merge/Private/SBlueprintMerge.cpp
@@ -124,6 +124,7 @@ void SBlueprintMerge::Construct(const FArguments& InArgs)
 	PanelLocal.Blueprint = InArgs._BlueprintLocal;
 
 	OwningEditor = InArgs._OwningEditor;
+	check( OwningEditor.IsValid() );
 
 	return SBlueprintDiff::Construct( InArgs._BaseArgs );
 }
@@ -303,49 +304,6 @@ TSharedRef<SWidget> SBlueprintMerge::GenerateToolbar()
 					.ToolTipText(LOCTEXT("CancelMergeToolTip", "Discards the changes made in the merge tool. Local file will remain unresolved so the merge tool can be run again if desired."))
 				]
 			]
-			+ SVerticalBox::Slot()
-			[
-				SNew(SButton)
-				.OnClicked(this, &SBlueprintMerge::OnTakeLocalClicked)
-				.Content()
-				[
-					SNew(STextBlock)
-					.Text(LOCTEXT("LocalChanges", "Take Local Changes Only"))
-					.ToolTipText(LOCTEXT("LocalChangesToolTip", "Resolve the conflict by discarding the changes made remotely. Conflict not fully resolved until the result is accepted."))
-				]
-			]
-			+ SVerticalBox::Slot()
-			[
-				SNew(SButton)
-				.OnClicked(this, &SBlueprintMerge::OnTakeRemoteClicked)
-				.Content()
-				[
-					SNew(STextBlock)
-					.Text(LOCTEXT("RemoteChanges", "Take Remote Changes Only"))
-					.ToolTipText(LOCTEXT("RemoteChangesToolTip", "Resolve the conflict by discarding the changes made locally. Conflict not fully resolved until the result is accepted."))
-				]
-			]
-			+ SVerticalBox::Slot()
-			[
-				SNew(SButton)
-				.OnClicked(this, &SBlueprintMerge::OnTakeBaseClicked)
-				.Content()
-				[
-					SNew(STextBlock)
-					.Text(LOCTEXT("BaseChanges", "Take Base Only (Revert All)"))
-					.ToolTipText(LOCTEXT("BaseChangesToolTip", "Resolve the conflict by discarding the both local and remote changes. Conflict not fully resolved until the result is accepted."))
-				]
-			]
-			+ SVerticalBox::Slot()
-			[
-				SNew(SButton)
-				.Content()
-				[
-					SNew(STextBlock)
-					.Text(LOCTEXT("AutoMerge", "Automatically Merge"))
-					.ToolTipText(LOCTEXT("AutoMergeChangesToolTip", "Resolve the conflict by attempting to automatically apply both local and remote changes. Conflict not fully resolved until the result is accepted."))
-				]
-			]
 		]
 		+ SVerticalBox::Slot()
 			.AutoHeight()
@@ -409,7 +367,7 @@ void SBlueprintMerge::OnSelectionChanged(SBlueprintDiff::FGraphToDiff Item, ESel
 	}
 
 	FString GraphName = Item->GetGraphOld()->GetName();
-	DisablePinDiffFocus(); // @todo doc: Not sure what this does..
+	DisablePinDiffFocus();
 	HandleGraphChanged(GraphName);
 
 	// Ensure that regenerated graphs are properly locked/unlocked:
@@ -514,6 +472,12 @@ TSharedRef< SWidget > SBlueprintMerge::GenerateDiffView(TArray<FDiffSingleResult
 	return Result;
 }
 
+UBlueprint* SBlueprintMerge::GetTargetBlueprint()
+{
+	// @todo doc, do we need to cache this on merge construction?
+	return OwningEditor->GetBlueprintObj();
+}
+
 FReply SBlueprintMerge::OnAcceptResultClicked()
 {
 	// Because merge operations are so destructive and can be confusing, lets write backups of the files involved:
@@ -522,22 +486,13 @@ FReply SBlueprintMerge::OnAcceptResultClicked()
 	WriteBackup(*PanelOld.Blueprint->GetOutermost(), BackupSubDir, TEXT("CommonBaseAsset") + FPackageName::GetAssetPackageExtension() );
 	WriteBackup(*PanelLocal.Blueprint->GetOutermost(), BackupSubDir, TEXT("LocalAsset") + FPackageName::GetAssetPackageExtension());
 
-	UPackage* Package = PanelLocal.Blueprint->GetOutermost();
+	UPackage* Package = GetTargetBlueprint()->GetOutermost();
 	TArray<UPackage*> PackagesToSave;
 	PackagesToSave.Add(Package);
 
 	// Perform the resolve with the SCC plugin, we do this first so that the editor doesn't warn about writing to a file that is unresolved:
 	ISourceControlModule::Get().GetProvider().Execute(ISourceControlOperation::Create<FResolve>(), SourceControlHelpers::PackageFilenames(PackagesToSave), EConcurrency::Synchronous);
 
-	// cr doc: this is too simple to be correct..
-	PanelLocal.Blueprint->RemoveGeneratedClasses();
-
-	auto Editor = OwningEditor.Pin();
-	if (Editor.IsValid())
-	{
-		Editor->CloseWindow();
-	}
-
 	const auto Result = FEditorFileUtils::PromptForCheckoutAndSave(PackagesToSave, /*bCheckDirty=*/ false, /*bPromptToSave=*/ false);
 	if (Result != FEditorFileUtils::PR_Success)
 	{
@@ -548,57 +503,15 @@ FReply SBlueprintMerge::OnAcceptResultClicked()
 		FSlateNotificationManager::Get().AddNotification( ErrorNotification );
 	}
 
-	return FReply::Handled();
-}
-
-FReply SBlueprintMerge::OnCancelClicked()
-{
-	auto Editor = OwningEditor.Pin();
-	if (Editor.IsValid())
-	{
-		Editor->CloseMergeTool();
-	}
-	return FReply::Handled();
-}
-
-FReply SBlueprintMerge::OnTakeLocalClicked()
-{
-	StageBlueprint(PanelLocal.Blueprint);
-	return FReply::Handled();
-}
+	OwningEditor->CloseMergeTool();
 
-FReply SBlueprintMerge::OnTakeRemoteClicked()
-{
-	StageBlueprint(PanelNew.Blueprint);
 	return FReply::Handled();
 }
 
-FReply SBlueprintMerge::OnTakeBaseClicked()
+FReply SBlueprintMerge::OnCancelClicked()
 {
-	StageBlueprint(PanelOld.Blueprint);
+	OwningEditor->CloseMergeTool();
 	return FReply::Handled();
 }
 
-void SBlueprintMerge::StageBlueprint(UBlueprint const* DesiredBP)
-{
-	// Take the DesiredBP and applies its changes to the result blueprint:
-	auto Editor = OwningEditor.Pin();
-	if( !Editor.IsValid() )
-	{
-		return;
-	}
-
-	UBlueprint* BlueprintResult = Editor->GetBlueprintObj();
-
-	// I want everything in the package... components, CDO, etc:
-
-
-	(void)PackagesAppearEquivalent(BlueprintResult->GetOutermost(), DesiredBP->GetOutermost());
-	check( PackagesAppearEquivalent( BlueprintResult->GetOutermost(), DesiredBP->GetOutermost() ) );
-
-	SGraphEditor::FGraphEditorEvents InEvents;
-	TArray< FGraphToDiff>  Selected = GraphsToDiff->GetSelectedItems();
-	check(Selected.Num() <= 1); // We set this control up to not support multiselect
-}
-
 #undef LOCTEXT_NAMESPACE
diff --git a/Engine/Source/Developer/Merge/Private/SBlueprintMerge.h b/Engine/Source/Developer/Merge/Private/SBlueprintMerge.h
--- a/Engine/Source/Developer/Merge/Private/SBlueprintMerge.h
+++ b/Engine/Source/Developer/Merge/Private/SBlueprintMerge.h
@@ -13,7 +13,7 @@ public:
 	SLATE_BEGIN_ARGS(SBlueprintMerge){}
 	SLATE_ARGUMENT(class UBlueprint*, BlueprintLocal)
 	SLATE_ARGUMENT( SBlueprintDiff::FArguments, BaseArgs )
-	SLATE_ARGUMENT( TWeakPtr<FBlueprintEditor>, OwningEditor )
+	SLATE_ARGUMENT( TSharedPtr<FBlueprintEditor>, OwningEditor )
 	SLATE_END_ARGS()
 	
 	void Construct(const FArguments& InArgs);
@@ -31,16 +31,13 @@ protected:
 	
 	/** Helper functions */
 	TSharedRef< SWidget > GenerateDiffView(TArray<FDiffSingleResult>& Diffs, TSharedPtr< const FUICommandInfo > CommandNext, TSharedPtr< const FUICommandInfo > CommandPrev, TArray< TSharedPtr<FDiffSingleResult> >& SharedResults);
+	UBlueprint* GetTargetBlueprint();
 
 	/** Event handlers */
 	FReply OnAcceptResultClicked();
 	FReply OnCancelClicked();
-	FReply OnTakeLocalClicked();
-	FReply OnTakeRemoteClicked();
-	FReply OnTakeBaseClicked();
-	void StageBlueprint(UBlueprint const* DesiredBP);
 
-	TWeakPtr<FBlueprintEditor> OwningEditor;
+	TSharedPtr<FBlueprintEditor> OwningEditor;
 	FDiffPanel PanelLocal;
 	TWeakPtr<SWindow>		OwningWindow;
 
