#UE4 pulling merge tool out of the tabmanager (fixes bug where we could not close the merge tool when the tab manager created it, fixes bug where tab manager would destroy merge tool when we switched tabs..)

diff --git a/Engine/Source/Developer/Merge/Private/Merge.cpp b/Engine/Source/Developer/Merge/Private/Merge.cpp
--- a/Engine/Source/Developer/Merge/Private/Merge.cpp
+++ b/Engine/Source/Developer/Merge/Private/Merge.cpp
@@ -10,6 +10,13 @@
 
 #define LOCTEXT_NAMESPACE "Merge"
 
+static void DisplayErrorMessage( const FText& ErrorMessage )
+{
+	FNotificationInfo Info(ErrorMessage);
+	Info.ExpireDuration = 5.0f;
+	FSlateNotificationManager::Get().AddNotification(Info);
+}
+
 static FSourceControlStatePtr GetSourceControlState(const FString& PackageName)
 {
 	ISourceControlProvider& SourceControlProvider = ISourceControlModule::Get().GetProvider();
@@ -49,17 +56,36 @@ static UObject* LoadRevision(const FString& AssetName, const ISourceControlRevis
 			}
 			else
 			{
-				// @todo DO: toast? i love toast..
+				DisplayErrorMessage( 
+					FText::Format(
+						LOCTEXT("MergedFailedToFindObject", 
+								"Aborted Load of {0} because we could not find an object named {1}" )
+								,FText::FromString(TempFileName)
+								,FText::FromString(AssetName) 
+					)
+				);
 			}
 		}
 		else
 		{
-			// @todo DO: toast? i love toast..
+			DisplayErrorMessage(
+				FText::Format(
+					LOCTEXT("MergedFailedToLoadPackage",
+							"Aborted Load of {0} because we could not load the package")
+							, FText::FromString(TempFileName)
+				)
+			);
 		}
 	}
 	else
 	{
-		// @todo DO: toast? i love toast..
+		DisplayErrorMessage(
+			FText::Format(
+				LOCTEXT("MergedFailedToFindRevision",
+						"Aborted Load of {0} because we could not get the requested revision")
+						, FText::FromString(TempFileName)
+			)
+		);
 	}
 	return NULL;
 }
@@ -70,40 +96,21 @@ static FRevisionInfo GetRevisionInfo(ISourceControlRevision const& FromRevision)
 	return Ret;
 }
 
-static UObject* LoadHeadRev(const FString& PackageName, const FString& AssetName, FRevisionInfo& OutRevInfo)
+static UObject* LoadHeadRev(const FString& PackageName, const FString& AssetName, const ISourceControlState& SourceControlState, FRevisionInfo& OutRevInfo)
 {
-	FSourceControlStatePtr SourceControlState = GetSourceControlState(PackageName);
-	if (SourceControlState.IsValid())
-	{
-		// HistoryItem(0) is apparently the head revision:
-		TSharedPtr<ISourceControlRevision, ESPMode::ThreadSafe> Revision = SourceControlState->GetHistoryItem(0);
-		check(Revision.IsValid());
-		OutRevInfo = GetRevisionInfo(*Revision);
-		return LoadRevision(AssetName, *Revision);
-	}
-	else
-	{
-		// @todo DO: toast? i love toast..
-	}
-
-	return NULL;
+	// HistoryItem(0) is apparently the head revision:
+	TSharedPtr<ISourceControlRevision, ESPMode::ThreadSafe> Revision = SourceControlState.GetHistoryItem(0);
+	check(Revision.IsValid());
+	OutRevInfo = GetRevisionInfo(*Revision);
+	return LoadRevision(AssetName, *Revision);
 }
 
-static UObject* LoadBaseRev(const FString& PackageName, const FString& AssetName, FRevisionInfo& OutRevInfo)
+static UObject* LoadBaseRev(const FString& PackageName, const FString& AssetName, const ISourceControlState& SourceControlState, FRevisionInfo& OutRevInfo)
 {
-	FSourceControlStatePtr SourceControlState = GetSourceControlState(PackageName);
-	if (SourceControlState.IsValid())
-	{
-		TSharedPtr<ISourceControlRevision, ESPMode::ThreadSafe> Revision = SourceControlState->GetBaseRevForMerge();
-		check(Revision.IsValid());
-		OutRevInfo = GetRevisionInfo(*Revision);
-		return LoadRevision(AssetName, *Revision);
-	}
-	else
-	{
-		// @todo DO: toast? i love toast..
-	}
-	return NULL;
+	TSharedPtr<ISourceControlRevision, ESPMode::ThreadSafe> Revision = SourceControlState.GetBaseRevForMerge();
+	check(Revision.IsValid());
+	OutRevInfo = GetRevisionInfo(*Revision);
+	return LoadRevision(AssetName, *Revision);
 }
 
 class FMerge : public IMerge
@@ -111,7 +118,7 @@ class FMerge : public IMerge
 	/** IModuleInterface implementation */
 	virtual void StartupModule() override;
 	virtual void ShutdownModule() override;
-	virtual TSharedPtr<SWidget> GenerateMergeWidget(const UBlueprint& Object, TSharedRef<FBlueprintEditor> Editor) override;
+	virtual TSharedRef<SWidget> GenerateMergeWidget(const UBlueprint& Object, TSharedRef<FBlueprintEditor> Editor) override;
 	virtual bool PendingMerge(const UBlueprint& BlueprintObj) const override;
 };
 IMPLEMENT_MODULE( FMerge, Merge )
@@ -127,7 +134,7 @@ void FMerge::ShutdownModule()
 	// we call this function before unloading the module.
 }
 
-TSharedPtr<SWidget> FMerge::GenerateMergeWidget(const UBlueprint& Object, TSharedRef<FBlueprintEditor> Editor)
+TSharedRef<SWidget> FMerge::GenerateMergeWidget(const UBlueprint& Object, TSharedRef<FBlueprintEditor> Editor)
 {
 	// merge the local asset with the depot, SCC provides us with the last common revision as
 	// a basis for the merge:
@@ -135,11 +142,28 @@ TSharedPtr<SWidget> FMerge::GenerateMergeWidget(const UBlueprint& Object, TShare
 	// @todo DO: this will probably need to be async.. pulling down some old versions of assets:
 	const FString& PackageName = Object.GetOutermost()->GetName();
 	const FString& AssetName = Object.GetName();
+	
+	FSourceControlStatePtr SourceControlState = GetSourceControlState(PackageName);
+	if (!SourceControlState.IsValid())
+	{
+		DisplayErrorMessage(
+			FText::Format(
+				LOCTEXT("MergeFailedNoSourceControl",
+						"Aborted Load of {0} from {1} because the source control state was invalidated")
+						, FText::FromString(AssetName)
+						, FText::FromString(PackageName)
+			)
+		);
+
+		return SNew(SHorizontalBox);
+	}
+
+	ISourceControlState const& SourceControlStateRef = *SourceControlState;
 
 	FRevisionInfo CurrentRevInfo = FRevisionInfo::InvalidRevision();
-	UBlueprint* RemoteBlueprint = Cast< UBlueprint >( LoadHeadRev(PackageName, AssetName, CurrentRevInfo) );
+	UBlueprint* RemoteBlueprint = Cast< UBlueprint >( LoadHeadRev(PackageName, AssetName, SourceControlStateRef, CurrentRevInfo) );
 	FRevisionInfo BaseRevInfo = FRevisionInfo::InvalidRevision();
-	UBlueprint* BaseBlueprint = Cast< UBlueprint >( LoadBaseRev(PackageName, AssetName, BaseRevInfo) );
+	UBlueprint* BaseBlueprint = Cast< UBlueprint >( LoadBaseRev(PackageName, AssetName, SourceControlStateRef, BaseRevInfo) );
 
 	if (RemoteBlueprint && BaseBlueprint)
 	{
@@ -172,11 +196,15 @@ TSharedPtr<SWidget> FMerge::GenerateMergeWidget(const UBlueprint& Object, TShare
 			MissingFiles = LOCTEXT("MergeBaseRevisionsFailed", "common base");
 		}
 
-		FText ErrorMessage = FText::Format(LOCTEXT("MergeRevisionLoadFailed", "Aborted Merge of {0} because we could not load {1}"), FText::FromString(Object.GetName()), MissingFiles);
-		FNotificationInfo Info(ErrorMessage);
-		Info.ExpireDuration = 3.0f;
-		FSlateNotificationManager::Get().AddNotification(Info);
-		return TSharedPtr<SWidget>();
+		DisplayErrorMessage(
+			FText::Format(
+				LOCTEXT("MergeRevisionLoadFailed", "Aborted Merge of {0} because we could not load {1}")
+				, FText::FromString(Object.GetName())
+				, MissingFiles
+			) 
+		);
+
+		return SNew(SHorizontalBox);
 	}
 }
 
diff --git a/Engine/Source/Developer/Merge/Private/SBlueprintMerge.cpp b/Engine/Source/Developer/Merge/Private/SBlueprintMerge.cpp
--- a/Engine/Source/Developer/Merge/Private/SBlueprintMerge.cpp
+++ b/Engine/Source/Developer/Merge/Private/SBlueprintMerge.cpp
@@ -124,7 +124,7 @@ void SBlueprintMerge::Construct(const FArguments& InArgs)
 	PanelLocal.Blueprint = InArgs._BlueprintLocal;
 
 	OwningEditor = InArgs._OwningEditor;
-	check( OwningEditor.IsValid() );
+	check( OwningEditor.Pin().IsValid() );
 
 	return SBlueprintDiff::Construct( InArgs._BaseArgs );
 }
@@ -475,7 +475,7 @@ TSharedRef< SWidget > SBlueprintMerge::GenerateDiffView(TArray<FDiffSingleResult
 UBlueprint* SBlueprintMerge::GetTargetBlueprint()
 {
 	// @todo doc, do we need to cache this on merge construction?
-	return OwningEditor->GetBlueprintObj();
+	return OwningEditor.Pin()->GetBlueprintObj();
 }
 
 FReply SBlueprintMerge::OnAcceptResultClicked()
@@ -503,14 +503,14 @@ FReply SBlueprintMerge::OnAcceptResultClicked()
 		FSlateNotificationManager::Get().AddNotification( ErrorNotification );
 	}
 
-	OwningEditor->CloseMergeTool();
+	OwningEditor.Pin()->CloseMergeTool();
 
 	return FReply::Handled();
 }
 
 FReply SBlueprintMerge::OnCancelClicked()
 {
-	OwningEditor->CloseMergeTool();
+	OwningEditor.Pin()->CloseMergeTool();
 	return FReply::Handled();
 }
 
diff --git a/Engine/Source/Developer/Merge/Private/SBlueprintMerge.h b/Engine/Source/Developer/Merge/Private/SBlueprintMerge.h
--- a/Engine/Source/Developer/Merge/Private/SBlueprintMerge.h
+++ b/Engine/Source/Developer/Merge/Private/SBlueprintMerge.h
@@ -37,7 +37,10 @@ protected:
 	FReply OnAcceptResultClicked();
 	FReply OnCancelClicked();
 
-	TSharedPtr<FBlueprintEditor> OwningEditor;
+	// We can't maintain a SharedPtr to OwningEditor because we don't logically own it,
+	// if you make this a SharedPtr closing the blueprint editor while the merge tool
+	// is open will prevent the blueprint editor destructor from running:
+	TWeakPtr<FBlueprintEditor> OwningEditor;
 	FDiffPanel PanelLocal;
 	TWeakPtr<SWindow>		OwningWindow;
 
diff --git a/Engine/Source/Developer/Merge/Public/Merge.h b/Engine/Source/Developer/Merge/Public/Merge.h
--- a/Engine/Source/Developer/Merge/Public/Merge.h
+++ b/Engine/Source/Developer/Merge/Public/Merge.h
@@ -48,7 +48,7 @@ public:
 	 *
 	 * @return The merge view widget
 	 */
-	virtual TSharedPtr<class SWidget> GenerateMergeWidget( const UBlueprint& Object, TSharedRef<FBlueprintEditor> Editor ) = 0;
+	virtual TSharedRef<class SWidget> GenerateMergeWidget( const UBlueprint& Object, TSharedRef<FBlueprintEditor> Editor ) = 0;
 
 	/** 
 	 * @return whether the blueprint is in a conflicted state, and can therefore be merged.
diff --git a/Engine/Source/Editor/Kismet/Private/BlueprintEditor.cpp b/Engine/Source/Editor/Kismet/Private/BlueprintEditor.cpp
--- a/Engine/Source/Editor/Kismet/Private/BlueprintEditor.cpp
+++ b/Engine/Source/Editor/Kismet/Private/BlueprintEditor.cpp
@@ -79,6 +79,8 @@
 
 #define LOCTEXT_NAMESPACE "BlueprintEditor"
 
+const FName MergeToolTabId = FName(TEXT("MergeTool"));
+
 /////////////////////////////////////////////////////
 // FSelectionDetailsSummoner
 
@@ -991,6 +993,10 @@ void FBlueprintEditor::CommonInitialization(const TArray<UBlueprint*>& InitBluep
 	// Make sure we know when tabs become active to update details tab
 	FGlobalTabmanager::Get()->OnActiveTabChanged_Subscribe( FOnActiveTabChanged::FDelegate::CreateRaw(this, &FBlueprintEditor::OnActiveTabChanged) );
 
+	FGlobalTabmanager::Get()->RegisterNomadTabSpawner(MergeToolTabId, FOnSpawnTab::CreateRaw(this, &FBlueprintEditor::CreateMergeToolTab))
+		.SetDisplayName(NSLOCTEXT("MergeTool", "TabTitle", "Merge Tool"))
+		.SetTooltipText(NSLOCTEXT("MergeTool", "TooltipText", "Open the Blueprint Merge Tool."));
+
 	if (InitBlueprints.Num() == 1)
 	{
 		// Load blueprint libraries
@@ -1496,12 +1502,15 @@ FBlueprintEditor::~FBlueprintEditor()
 		Editor->UnregisterForUndo( this );
 	}
 
+	CloseMergeTool();
+
 	if (GetBlueprintObj())
 	{
 		GetBlueprintObj()->OnChanged().RemoveAll( this );
 	}
 
 	FGlobalTabmanager::Get()->OnActiveTabChanged_Unsubscribe( FOnActiveTabChanged::FDelegate::CreateRaw(this, &FBlueprintEditor::OnActiveTabChanged) );
+	FGlobalTabmanager::Get()->UnregisterNomadTabSpawner(MergeToolTabId);
 
 	if (FEngineAnalytics::IsAvailable())
 	{
@@ -1865,7 +1874,7 @@ void FBlueprintEditor::CreateDefaultCommands()
 		FCanExecuteAction::CreateStatic( &FLocalKismetCallbacks::CanRecompileModules ));
 
 	ToolkitCommands->MapAction(FBlueprintEditorCommands::Get().BeginBlueprintMerge,
-		FExecuteAction::CreateSP(this, &FBlueprintEditor::OnBeginBlueprintMerge),
+		FExecuteAction::CreateStatic([] { FGlobalTabmanager::Get()->InvokeTab(MergeToolTabId); } ),
 		FCanExecuteAction());
 }
 
@@ -4652,7 +4661,7 @@ void FBlueprintEditor::UnregisterSCSEditorCustomization(const FName& InComponent
 	SCSEditorCustomizations.Remove(InComponentName);
 }
 
-TSharedRef<SWidget> FBlueprintEditor::ShowMergeTool()
+TSharedRef<SDockTab> FBlueprintEditor::CreateMergeToolTab(const FSpawnTabArgs&)
 {
 	/** 
 	 * This function has to handle three case:
@@ -4660,33 +4669,26 @@ TSharedRef<SWidget> FBlueprintEditor::ShowMergeTool()
 	 *	2. The merge tool is being restored, if we have a pending merge display the merge tool, otherwise display an empty widget
 	 *  3. The merge tool is being invoked, display the merge tool 
 	 */
-	auto MergeToolPtr = MergeTool.Pin();
-	if( !MergeToolPtr.IsValid() )
+	TSharedPtr<SDockTab> Ret = MergeTool.Pin();
+	if( !Ret.IsValid() )
 	{
-		if( IMerge::Get().PendingMerge(*GetBlueprintObj()) )
-		{
-			MergeToolPtr = IMerge::Get().GenerateMergeWidget(*GetBlueprintObj(), SharedThis(this) ).ToSharedRef();
-		}
-		else
-		{
-			// We *have* to return a valid reference to a widget, so return an empty box. This is arguabla a deficiency in
-			// our controlling widgets that we can't 'restore' with nothing.
-			MergeToolPtr = SNew(SHorizontalBox);
-		}
-
-		// Keep a weak reference to the widget we've created, widget will be owned by the docktab:
-		MergeTool = MergeToolPtr;
+		Ret = SNew(SDockTab)
+			.Label(LOCTEXT("MergeTool_Title", "Merge Tool"))
+			[
+				IMerge::Get().GenerateMergeWidget(*GetBlueprintObj(), SharedThis(this))
+			];
+		MergeTool = Ret;
 	}
 
-	return MergeToolPtr.ToSharedRef();
+	return Ret.ToSharedRef();
 }
 
 void FBlueprintEditor::CloseMergeTool()
 {
-	auto MergeToolDockTabPtr = MergeToolDockTab.Pin();
-	if( MergeToolDockTabPtr.IsValid() )
+	auto MergeToolPtr = MergeTool.Pin();
+	if( MergeToolPtr.IsValid() )
 	{
-		MergeToolDockTabPtr->RequestCloseTab();
+		MergeToolPtr->RequestCloseTab();
 	}
 }
 
@@ -5508,11 +5510,6 @@ void FBlueprintEditor::OnRepairCorruptedBlueprint()
 	Compiler.RecoverCorruptedBlueprint(GetBlueprintObj());
 }
 
-void FBlueprintEditor::OnBeginBlueprintMerge()
-{
-	MergeToolDockTab = TabManager->InvokeTab(FBlueprintEditorTabs::MergeToolID);
-}
-
 void FBlueprintEditor::StartEditingDefaults(bool bAutoFocus, bool bForceRefresh)
 {
 	if (IsEditingSingleBlueprint())
diff --git a/Engine/Source/Editor/Kismet/Private/BlueprintEditorModes.cpp b/Engine/Source/Editor/Kismet/Private/BlueprintEditorModes.cpp
--- a/Engine/Source/Editor/Kismet/Private/BlueprintEditorModes.cpp
+++ b/Engine/Source/Editor/Kismet/Private/BlueprintEditorModes.cpp
@@ -52,7 +52,7 @@ FBlueprintEditorApplicationMode::FBlueprintEditorApplicationMode(TSharedPtr<clas
 	BlueprintEditorTabFactories.RegisterFactory(MakeShareable(new FMyBlueprintSummoner(InBlueprintEditor)));
 	BlueprintEditorTabFactories.RegisterFactory(MakeShareable(new FCompilerResultsSummoner(InBlueprintEditor)));
 	BlueprintEditorTabFactories.RegisterFactory(MakeShareable(new FFindResultsSummoner(InBlueprintEditor)));
-	BlueprintEditorTabFactories.RegisterFactory(MakeShareable(new FMergeToolSummoner(InBlueprintEditor)));
+	
 	if( bRegisterViewport )
 	{
 		BlueprintEditorTabFactories.RegisterFactory(MakeShareable(new FSCSViewportSummoner(InBlueprintEditor)));
@@ -184,7 +184,6 @@ FBlueprintDefaultsApplicationMode::FBlueprintDefaultsApplicationMode(TSharedPtr<
 	
 	BlueprintDefaultsTabFactories.RegisterFactory(MakeShareable(new FDefaultsEditorSummoner(InBlueprintEditor)));
 	BlueprintDefaultsTabFactories.RegisterFactory(MakeShareable(new FFindResultsSummoner(InBlueprintEditor)));
-	BlueprintDefaultsTabFactories.RegisterFactory(MakeShareable(new FMergeToolSummoner(InBlueprintEditor)));
 
 	TabLayout = FTabManager::NewLayout( "Standalone_BlueprintDefaults_Layout_v4" )
 		->AddArea
@@ -253,7 +252,6 @@ FBlueprintComponentsApplicationMode::FBlueprintComponentsApplicationMode(TShared
 	BlueprintComponentsTabFactories.RegisterFactory(MakeShareable(new FSelectionDetailsSummoner(InBlueprintEditor)));
 	BlueprintComponentsTabFactories.RegisterFactory(MakeShareable(new FDefaultsEditorSummoner(InBlueprintEditor)));
 	BlueprintComponentsTabFactories.RegisterFactory(MakeShareable(new FFindResultsSummoner(InBlueprintEditor)));
-	BlueprintComponentsTabFactories.RegisterFactory(MakeShareable(new FMergeToolSummoner(InBlueprintEditor)));
 
 	TabLayout = FTabManager::NewLayout( "Standalone_BlueprintComponents_Layout_v5" )
 		->AddArea
diff --git a/Engine/Source/Editor/Kismet/Private/BlueprintEditorTabFactories.h b/Engine/Source/Editor/Kismet/Private/BlueprintEditorTabFactories.h
--- a/Engine/Source/Editor/Kismet/Private/BlueprintEditorTabFactories.h
+++ b/Engine/Source/Editor/Kismet/Private/BlueprintEditorTabFactories.h
@@ -488,31 +488,4 @@ public:
 	}
 };
 
-//////////////////////////////////////////////////////////////////////////
-// FMergeToolSummoner
-
-struct FMergeToolSummoner : public FWorkflowTabFactory
-{
-	FMergeToolSummoner(TSharedPtr<class FAssetEditorToolkit> InHostingApp)
-		: FWorkflowTabFactory(FBlueprintEditorTabs::MergeToolID, InHostingApp)
-	{
-		TabLabel = LOCTEXT("MergeToolTabTitle", "Merge Tool");
-		bIsSingleton = true;
-
-		ViewMenuDescription = LOCTEXT("MergeToolView", "Merge Tool");
-		ViewMenuTooltip = LOCTEXT("MergeToolView_ToolTip", "Displays three revisions of a blueprint for merging purposes, provides access to SSC merge operations");
-	}
-
-	virtual TSharedRef<SWidget> CreateTabBody(const FWorkflowTabSpawnInfo& Info) const override
-	{
-		TSharedPtr<FBlueprintEditor> BlueprintEditorPtr = StaticCastSharedPtr<FBlueprintEditor>(HostingApp.Pin());
-		return BlueprintEditorPtr->ShowMergeTool();
-	}
-
-	virtual FText GetTabToolTipText(const FWorkflowTabSpawnInfo& Info) const override
-	{
-		return LOCTEXT("MergeToolTooltip", "The merge tool allows users to merge two blueprints that diverged from a common base blueprint");
-	}
-};
-
 #undef LOCTEXT_NAMESPACE
diff --git a/Engine/Source/Editor/Kismet/Public/BlueprintEditor.h b/Engine/Source/Editor/Kismet/Public/BlueprintEditor.h
--- a/Engine/Source/Editor/Kismet/Public/BlueprintEditor.h
+++ b/Engine/Source/Editor/Kismet/Public/BlueprintEditor.h
@@ -451,7 +451,7 @@ public:
 	/** 
 	 * Forces the merge tool to be shown, if no merge is possible, returns an empty widget:
 	 */
-	TSharedRef<SWidget> ShowMergeTool();
+	TSharedRef<SDockTab> CreateMergeToolTab(const FSpawnTabArgs&);
 	
 	/** 
 	 * Closes the merge tool, rather than simply hiding it.
@@ -714,8 +714,6 @@ protected:
 	void OnListObjectsReferencedByBlueprint();
 	void OnRepairCorruptedBlueprint();
 
-	void OnBeginBlueprintMerge();
-
 	void OnNodeDoubleClicked(class UEdGraphNode* Node);
 
 	virtual void OnEditTabClosed(TSharedRef<SDockTab> Tab);
@@ -924,11 +922,8 @@ protected:
 	/** Find results log as well as the search filter */
 	TSharedPtr<class SFindInBlueprints> FindResults;
 
-	/** Merge tool */
-	TWeakPtr<class SWidget> MergeTool;
-
-	/** Docktab that contains the merge tool */
-	TWeakPtr<class SDockTab> MergeToolDockTab;
+	/** Merge tool - WeakPtr because it's owned by the GlobalTabManager */
+	TWeakPtr<class SDockTab> MergeTool;
 
 	/** Reference to owner of the current popup */
 	TWeakPtr<class SWindow> NameEntryPopupWindow;
