#UE4 Reworked tab spawner registration - limiting ourself to a single merge operation at a time for the moment.

diff --git a/Engine/Source/Developer/Merge/Private/Merge.cpp b/Engine/Source/Developer/Merge/Private/Merge.cpp
--- a/Engine/Source/Developer/Merge/Private/Merge.cpp
+++ b/Engine/Source/Developer/Merge/Private/Merge.cpp
@@ -10,6 +10,8 @@
 
 #define LOCTEXT_NAMESPACE "Merge"
 
+const FName MergeToolTabId = FName(TEXT("MergeTool"));
+
 static void DisplayErrorMessage( const FText& ErrorMessage )
 {
 	FNotificationInfo Info(ErrorMessage);
@@ -118,24 +120,45 @@ class FMerge : public IMerge
 	/** IModuleInterface implementation */
 	virtual void StartupModule() override;
 	virtual void ShutdownModule() override;
-	virtual TSharedRef<SWidget> GenerateMergeWidget(const UBlueprint& Object, TSharedRef<FBlueprintEditor> Editor) override;
+	virtual TSharedRef<SDockTab> GenerateMergeWidget(const UBlueprint& Object, TSharedRef<FBlueprintEditor> Editor) override;
 	virtual bool PendingMerge(const UBlueprint& BlueprintObj) const override;
+
+	// Simplest to only allow one merge operation at a time, we could easily make this a map of Blueprint=>MergeTab
+	// but doing so will complicate the tab management
+	TWeakPtr<SDockTab> ActiveTab;
 };
 IMPLEMENT_MODULE( FMerge, Merge )
 
 void FMerge::StartupModule()
 {
 	// This code will execute after your module is loaded into memory (but after global variables are initialized, of course.)
+
+	FGlobalTabmanager::Get()->RegisterNomadTabSpawner(
+		MergeToolTabId, 
+		FOnSpawnTab::CreateStatic( [](const FSpawnTabArgs&) { return SNew(SDockTab); } ) )
+		.SetDisplayName(NSLOCTEXT("MergeTool", "TabTitle", "Merge Tool"))
+		.SetTooltipText(NSLOCTEXT("MergeTool", "TooltipText", "Used to display several versions of a blueprint that need to be merged into a single version.")
+	);
 }
 
 void FMerge::ShutdownModule()
 {
 	// This function may be called during shutdown to clean up your module.  For modules that support dynamic reloading,
 	// we call this function before unloading the module.
+	FGlobalTabmanager::Get()->UnregisterNomadTabSpawner(MergeToolTabId);
 }
 
-TSharedRef<SWidget> FMerge::GenerateMergeWidget(const UBlueprint& Object, TSharedRef<FBlueprintEditor> Editor)
+TSharedRef<SDockTab> FMerge::GenerateMergeWidget(const UBlueprint& Object, TSharedRef<FBlueprintEditor> Editor)
 {
+	auto ActiveTabPtr = ActiveTab.Pin();
+	if( ActiveTabPtr.IsValid() )
+	{
+		// just bring the tab to the foreground:
+		auto CurrentTab = FGlobalTabmanager::Get()->InvokeTab(MergeToolTabId);
+		check( CurrentTab == ActiveTabPtr );
+		return ActiveTabPtr.ToSharedRef();
+	}
+
 	// merge the local asset with the depot, SCC provides us with the last common revision as
 	// a basis for the merge:
 
@@ -143,6 +166,8 @@ TSharedRef<SWidget> FMerge::GenerateMergeWidget(const UBlueprint& Object, TShare
 	const FString& PackageName = Object.GetOutermost()->GetName();
 	const FString& AssetName = Object.GetName();
 	
+	TSharedPtr<SWidget> Contents;
+
 	FSourceControlStatePtr SourceControlState = GetSourceControlState(PackageName);
 	if (!SourceControlState.IsValid())
 	{
@@ -155,57 +180,65 @@ TSharedRef<SWidget> FMerge::GenerateMergeWidget(const UBlueprint& Object, TShare
 			)
 		);
 
-		return SNew(SHorizontalBox);
-	}
-
-	ISourceControlState const& SourceControlStateRef = *SourceControlState;
-
-	FRevisionInfo CurrentRevInfo = FRevisionInfo::InvalidRevision();
-	UBlueprint* RemoteBlueprint = Cast< UBlueprint >( LoadHeadRev(PackageName, AssetName, SourceControlStateRef, CurrentRevInfo) );
-	FRevisionInfo BaseRevInfo = FRevisionInfo::InvalidRevision();
-	UBlueprint* BaseBlueprint = Cast< UBlueprint >( LoadBaseRev(PackageName, AssetName, SourceControlStateRef, BaseRevInfo) );
-
-	if (RemoteBlueprint && BaseBlueprint)
-	{
-		FMergeDisplayArgs DisplayArgs = { CurrentRevInfo, BaseRevInfo };
-
-		SBlueprintDiff::FArguments BaseArgs;
-		BaseArgs.BlueprintOld(BaseBlueprint)
-			.OldRevision(CurrentRevInfo)
-			.BlueprintNew(RemoteBlueprint)
-			.NewRevision(BaseRevInfo);
-
-		return SNew(SBlueprintMerge)
-			.BlueprintLocal(const_cast<UBlueprint*>(&Object))
-			.OwningEditor(Editor)
-			.BaseArgs(BaseArgs);
+		Contents = SNew(SHorizontalBox);
 	}
 	else
 	{
-		FText MissingFiles;
-		if (!RemoteBlueprint && !BaseBlueprint)
-		{
-			MissingFiles = LOCTEXT("MergeBothRevisionsFailed", "common base, nor conflicting revision");
-		}
-		else if (!RemoteBlueprint)
+		ISourceControlState const& SourceControlStateRef = *SourceControlState;
+
+		FRevisionInfo CurrentRevInfo = FRevisionInfo::InvalidRevision();
+		UBlueprint* RemoteBlueprint = Cast< UBlueprint >(LoadHeadRev(PackageName, AssetName, SourceControlStateRef, CurrentRevInfo));
+		FRevisionInfo BaseRevInfo = FRevisionInfo::InvalidRevision();
+		UBlueprint* BaseBlueprint = Cast< UBlueprint >(LoadBaseRev(PackageName, AssetName, SourceControlStateRef, BaseRevInfo));
+
+		if (RemoteBlueprint && BaseBlueprint)
 		{
-			MissingFiles = LOCTEXT("MergeConflictingRevisionsFailed", "conflicting revision");
+			FMergeDisplayArgs DisplayArgs = { CurrentRevInfo, BaseRevInfo };
+
+			SBlueprintDiff::FArguments BaseArgs;
+			BaseArgs.BlueprintOld(BaseBlueprint)
+				.OldRevision(CurrentRevInfo)
+				.BlueprintNew(RemoteBlueprint)
+				.NewRevision(BaseRevInfo);
+
+			Contents = SNew(SBlueprintMerge)
+				.BlueprintLocal(const_cast<UBlueprint*>(&Object))
+				.OwningEditor(Editor)
+				.BaseArgs(BaseArgs);
 		}
 		else
 		{
-			MissingFiles = LOCTEXT("MergeBaseRevisionsFailed", "common base");
-		}
+			FText MissingFiles;
+			if (!RemoteBlueprint && !BaseBlueprint)
+			{
+				MissingFiles = LOCTEXT("MergeBothRevisionsFailed", "common base, nor conflicting revision");
+			}
+			else if (!RemoteBlueprint)
+			{
+				MissingFiles = LOCTEXT("MergeConflictingRevisionsFailed", "conflicting revision");
+			}
+			else
+			{
+				MissingFiles = LOCTEXT("MergeBaseRevisionsFailed", "common base");
+			}
 
-		DisplayErrorMessage(
-			FText::Format(
+			DisplayErrorMessage(
+				FText::Format(
 				LOCTEXT("MergeRevisionLoadFailed", "Aborted Merge of {0} because we could not load {1}")
 				, FText::FromString(Object.GetName())
 				, MissingFiles
-			) 
-		);
+				)
+				);
 
-		return SNew(SHorizontalBox);
+			Contents = SNew(SHorizontalBox);
+		}
 	}
+
+	TSharedRef<SDockTab> Tab =  FGlobalTabmanager::Get()->InvokeTab(MergeToolTabId);
+	Tab->SetContent(Contents.ToSharedRef());
+	ActiveTab = Tab;
+	return Tab;
+
 }
 
 bool FMerge::PendingMerge(const UBlueprint& BlueprintObj) const
diff --git a/Engine/Source/Developer/Merge/Public/Merge.h b/Engine/Source/Developer/Merge/Public/Merge.h
--- a/Engine/Source/Developer/Merge/Public/Merge.h
+++ b/Engine/Source/Developer/Merge/Public/Merge.h
@@ -48,7 +48,7 @@ public:
 	 *
 	 * @return The merge view widget
 	 */
-	virtual TSharedRef<class SWidget> GenerateMergeWidget( const UBlueprint& Object, TSharedRef<FBlueprintEditor> Editor ) = 0;
+	virtual TSharedRef<class SDockTab> GenerateMergeWidget( const UBlueprint& Object, TSharedRef<FBlueprintEditor> Editor ) = 0;
 
 	/** 
 	 * @return whether the blueprint is in a conflicted state, and can therefore be merged.
diff --git a/Engine/Source/Editor/Kismet/Private/BlueprintEditor.cpp b/Engine/Source/Editor/Kismet/Private/BlueprintEditor.cpp
--- a/Engine/Source/Editor/Kismet/Private/BlueprintEditor.cpp
+++ b/Engine/Source/Editor/Kismet/Private/BlueprintEditor.cpp
@@ -79,8 +79,6 @@
 
 #define LOCTEXT_NAMESPACE "BlueprintEditor"
 
-const FName MergeToolTabId = FName(TEXT("MergeTool"));
-
 /////////////////////////////////////////////////////
 // FSelectionDetailsSummoner
 
@@ -1869,7 +1867,7 @@ void FBlueprintEditor::CreateDefaultCommands()
 		FCanExecuteAction::CreateStatic( &FLocalKismetCallbacks::CanRecompileModules ));
 
 	ToolkitCommands->MapAction(FBlueprintEditorCommands::Get().BeginBlueprintMerge,
-		FExecuteAction::CreateStatic([] { FGlobalTabmanager::Get()->InvokeTab(MergeToolTabId); } ),
+		FExecuteAction::CreateSP(this, &FBlueprintEditor::CreateMergeToolTab),
 		FCanExecuteAction());
 }
 
@@ -4656,26 +4654,9 @@ void FBlueprintEditor::UnregisterSCSEditorCustomization(const FName& InComponent
 	SCSEditorCustomizations.Remove(InComponentName);
 }
 
-TSharedRef<SDockTab> FBlueprintEditor::CreateMergeToolTab(const FSpawnTabArgs&)
+void FBlueprintEditor::CreateMergeToolTab()
 {
-	/** 
-	 * This function has to handle three case:
-	 *	1. The merge tool is already created, return it
-	 *	2. The merge tool is being restored, if we have a pending merge display the merge tool, otherwise display an empty widget
-	 *  3. The merge tool is being invoked, display the merge tool 
-	 */
-	TSharedPtr<SDockTab> Ret = MergeTool.Pin();
-	if( !Ret.IsValid() )
-	{
-		Ret = SNew(SDockTab)
-			.Label(LOCTEXT("MergeTool_Title", "Merge Tool"))
-			[
-				IMerge::Get().GenerateMergeWidget(*GetBlueprintObj(), SharedThis(this))
-			];
-		MergeTool = Ret;
-	}
-
-	return Ret.ToSharedRef();
+	MergeTool = IMerge::Get().GenerateMergeWidget(*GetBlueprintObj(), SharedThis(this));
 }
 
 void FBlueprintEditor::CloseMergeTool()
diff --git a/Engine/Source/Editor/Kismet/Public/BlueprintEditor.h b/Engine/Source/Editor/Kismet/Public/BlueprintEditor.h
--- a/Engine/Source/Editor/Kismet/Public/BlueprintEditor.h
+++ b/Engine/Source/Editor/Kismet/Public/BlueprintEditor.h
@@ -448,14 +448,10 @@ public:
 	 */
 	void UnregisterSCSEditorCustomization(const FName& InComponentName);
 
-	/** 
-	 * Forces the merge tool to be shown, if no merge is possible, returns an empty widget:
-	 */
-	TSharedRef<SDockTab> CreateMergeToolTab(const FSpawnTabArgs&);
+	/** Forces the merge tool to be shown */
+	void CreateMergeToolTab();
 	
-	/** 
-	 * Closes the merge tool, rather than simply hiding it.
-	 */
+	/** Closes the merge tool, rather than simply hiding it */
 	void CloseMergeTool();
 
 	/** 
