#ue4 const correctness for merge and diff tools... merge copies the blueprint it is merging now, so that we can present a fixed view of the blueprint at the beginning of the merge operation

diff --git a/Engine/Source/Developer/Merge/Private/Merge.cpp b/Engine/Source/Developer/Merge/Private/Merge.cpp
--- a/Engine/Source/Developer/Merge/Private/Merge.cpp
+++ b/Engine/Source/Developer/Merge/Private/Merge.cpp
@@ -60,10 +60,9 @@ static UObject* LoadRevision(const FString& AssetName, const ISourceControlRevis
 			{
 				DisplayErrorMessage( 
 					FText::Format(
-						LOCTEXT("MergedFailedToFindObject", 
-								"Aborted Load of {0} because we could not find an object named {1}" )
-								,FText::FromString(TempFileName)
-								,FText::FromString(AssetName) 
+						LOCTEXT("MergedFailedToFindObject", "Aborted Load of {0} because we could not find an object named {1}" )
+						, FText::FromString(TempFileName)
+						, FText::FromString(AssetName) 
 					)
 				);
 			}
@@ -72,9 +71,8 @@ static UObject* LoadRevision(const FString& AssetName, const ISourceControlRevis
 		{
 			DisplayErrorMessage(
 				FText::Format(
-					LOCTEXT("MergedFailedToLoadPackage",
-							"Aborted Load of {0} because we could not load the package")
-							, FText::FromString(TempFileName)
+					LOCTEXT("MergedFailedToLoadPackage", "Aborted Load of {0} because we could not load the package")
+					, FText::FromString(TempFileName)
 				)
 			);
 		}
@@ -83,9 +81,8 @@ static UObject* LoadRevision(const FString& AssetName, const ISourceControlRevis
 	{
 		DisplayErrorMessage(
 			FText::Format(
-				LOCTEXT("MergedFailedToFindRevision",
-						"Aborted Load of {0} because we could not get the requested revision")
-						, FText::FromString(TempFileName)
+				LOCTEXT("MergedFailedToFindRevision", "Aborted Load of {0} because we could not get the requested revision")
+				, FText::FromString(TempFileName)
 			)
 		);
 	}
@@ -173,10 +170,9 @@ TSharedRef<SDockTab> FMerge::GenerateMergeWidget(const UBlueprint& Object, TShar
 	{
 		DisplayErrorMessage(
 			FText::Format(
-				LOCTEXT("MergeFailedNoSourceControl",
-						"Aborted Load of {0} from {1} because the source control state was invalidated")
-						, FText::FromString(AssetName)
-						, FText::FromString(PackageName)
+				LOCTEXT("MergeFailedNoSourceControl", "Aborted Load of {0} from {1} because the source control state was invalidated")
+				, FText::FromString(AssetName)
+				, FText::FromString(PackageName)
 			)
 		);
 
@@ -187,9 +183,9 @@ TSharedRef<SDockTab> FMerge::GenerateMergeWidget(const UBlueprint& Object, TShar
 		ISourceControlState const& SourceControlStateRef = *SourceControlState;
 
 		FRevisionInfo CurrentRevInfo = FRevisionInfo::InvalidRevision();
-		UBlueprint* RemoteBlueprint = Cast< UBlueprint >(LoadHeadRev(PackageName, AssetName, SourceControlStateRef, CurrentRevInfo));
+		const UBlueprint* RemoteBlueprint = Cast< UBlueprint >(LoadHeadRev(PackageName, AssetName, SourceControlStateRef, CurrentRevInfo));
 		FRevisionInfo BaseRevInfo = FRevisionInfo::InvalidRevision();
-		UBlueprint* BaseBlueprint = Cast< UBlueprint >(LoadBaseRev(PackageName, AssetName, SourceControlStateRef, BaseRevInfo));
+		const UBlueprint* BaseBlueprint = Cast< UBlueprint >(LoadBaseRev(PackageName, AssetName, SourceControlStateRef, BaseRevInfo));
 
 		if (RemoteBlueprint && BaseBlueprint)
 		{
@@ -202,7 +198,7 @@ TSharedRef<SDockTab> FMerge::GenerateMergeWidget(const UBlueprint& Object, TShar
 				.NewRevision(BaseRevInfo);
 
 			Contents = SNew(SBlueprintMerge)
-				.BlueprintLocal(const_cast<UBlueprint*>(&Object))
+				.BlueprintLocal(static_cast<const UBlueprint*> ( StaticDuplicateObject( &Object, GetTransientPackage(), TEXT("None") ) ) )
 				.OwningEditor(Editor)
 				.BaseArgs(BaseArgs);
 		}
@@ -224,11 +220,11 @@ TSharedRef<SDockTab> FMerge::GenerateMergeWidget(const UBlueprint& Object, TShar
 
 			DisplayErrorMessage(
 				FText::Format(
-				LOCTEXT("MergeRevisionLoadFailed", "Aborted Merge of {0} because we could not load {1}")
-				, FText::FromString(Object.GetName())
-				, MissingFiles
+					LOCTEXT("MergeRevisionLoadFailed", "Aborted Merge of {0} because we could not load {1}")
+					, FText::FromString(Object.GetName())
+					, MissingFiles
 				)
-				);
+			);
 
 			Contents = SNew(SHorizontalBox);
 		}
diff --git a/Engine/Source/Developer/Merge/Private/SBlueprintMerge.h b/Engine/Source/Developer/Merge/Private/SBlueprintMerge.h
--- a/Engine/Source/Developer/Merge/Private/SBlueprintMerge.h
+++ b/Engine/Source/Developer/Merge/Private/SBlueprintMerge.h
@@ -11,7 +11,7 @@ public:
 	SBlueprintMerge();
 
 	SLATE_BEGIN_ARGS(SBlueprintMerge){}
-	SLATE_ARGUMENT(class UBlueprint*, BlueprintLocal)
+	SLATE_ARGUMENT( const class UBlueprint*, BlueprintLocal)
 	SLATE_ARGUMENT( SBlueprintDiff::FArguments, BaseArgs )
 	SLATE_ARGUMENT( TSharedPtr<FBlueprintEditor>, OwningEditor )
 	SLATE_END_ARGS()
diff --git a/Engine/Source/Editor/Kismet/Public/SBlueprintDiff.h b/Engine/Source/Editor/Kismet/Public/SBlueprintDiff.h
--- a/Engine/Source/Editor/Kismet/Public/SBlueprintDiff.h
+++ b/Engine/Source/Editor/Kismet/Public/SBlueprintDiff.h
@@ -150,8 +150,8 @@ public:
 	DECLARE_DELEGATE_TwoParams( FOpenInDefaults, const class UBlueprint* , const class UBlueprint* );
 
 	SLATE_BEGIN_ARGS( SBlueprintDiff ){}
-			SLATE_ARGUMENT( class UBlueprint*, BlueprintOld )
-			SLATE_ARGUMENT(class UBlueprint*, BlueprintNew )
+			SLATE_ARGUMENT( const class UBlueprint*, BlueprintOld )
+			SLATE_ARGUMENT( const class UBlueprint*, BlueprintNew )
 			SLATE_ARGUMENT( struct FRevisionInfo, OldRevision )
 			SLATE_ARGUMENT( struct FRevisionInfo, NewRevision )
 			SLATE_ARGUMENT( bool, ShowAssetNames )
