#UE4 blueprintmerge no longer has buttons, binds to blueprint editor. No longer inherits from SBlueprintDiff.

diff --git a/Engine/Source/Developer/Merge/Private/Merge.cpp b/Engine/Source/Developer/Merge/Private/Merge.cpp
--- a/Engine/Source/Developer/Merge/Private/Merge.cpp
+++ b/Engine/Source/Developer/Merge/Private/Merge.cpp
@@ -189,18 +189,14 @@ TSharedRef<SDockTab> FMerge::GenerateMergeWidget(const UBlueprint& Object, TShar
 
 		if (RemoteBlueprint && BaseBlueprint)
 		{
-			FMergeDisplayArgs DisplayArgs = { CurrentRevInfo, BaseRevInfo };
-
-			SBlueprintDiff::FArguments BaseArgs;
-			BaseArgs.BlueprintOld(BaseBlueprint)
-				.OldRevision(CurrentRevInfo)
-				.BlueprintNew(RemoteBlueprint)
-				.NewRevision(BaseRevInfo);
-
-			Contents = SNew(SBlueprintMerge)
-				.BlueprintLocal(static_cast<const UBlueprint*> ( StaticDuplicateObject( &Object, GetTransientPackage(), TEXT("None") ) ) )
-				.OwningEditor(Editor)
-				.BaseArgs(BaseArgs);
+			FBlueprintMergeData Data(Editor 
+									, static_cast<const UBlueprint*>(StaticDuplicateObject(&Object, GetTransientPackage(), TEXT("None")))
+									, BaseBlueprint
+									, CurrentRevInfo
+									, RemoteBlueprint
+									, BaseRevInfo );
+
+			Contents = SNew(SBlueprintMerge, Data);
 		}
 		else
 		{
diff --git a/Engine/Source/Developer/Merge/Private/SBlueprintMerge.cpp b/Engine/Source/Developer/Merge/Private/SBlueprintMerge.cpp
--- a/Engine/Source/Developer/Merge/Private/SBlueprintMerge.cpp
+++ b/Engine/Source/Developer/Merge/Private/SBlueprintMerge.cpp
@@ -1,490 +1,582 @@
 // Copyright 1998-2014 Epic Games, Inc. All Rights Reserved.
 
 #include "MergePrivatePCH.h"
-#include "SBlueprintMerge.h"
+#include "BlueprintEditor.h"
 #include "DiffResults.h"
 #include "GraphDiffControl.h"
 #include "ISourceControlModule.h"
-#include "BlueprintEditor.h"
+#include "SBlueprintDiff.h"
+#include "SBlueprintMerge.h"
 
 #define LOCTEXT_NAMESPACE "SBlueprintMerge"
 
-// todo doc: where doe this utility function belong!
-static UPackage* CreateTempPackage(FString Name)
-{
-	static uint32 TempUid = 0;
-	FString TempPackageName = FString::Printf(TEXT("/Temp/BpMerge-%u-%s"), TempUid++, *Name);
-	return CreatePackage(NULL, *TempPackageName);
-}
-
-template< class Predicate >
-UObject* FindOuter( UObject& Obj, Predicate P )
+// This cludge is used to move arrays of structures into arrays of 
+// shared ptrs of those structures. It is used because SListView
+// does not support arrays of values:
+template< typename T >
+TArray< TSharedPtr<T> > HackToShared( const TArray<T>& Values )
 {
-	UObject* Iter = &Obj;
-	while( Iter )
+	TArray< TSharedPtr<T> > Ret;
+	for( const auto& Value : Values )
 	{
-		const UObject* ImmutableIter = Iter;
-		if( P(*ImmutableIter) )
-		{
-			break;
-		}
-		Iter = Iter->GetOuter();
+		Ret.Add( TSharedPtr<T>( new T(Value) ) );
 	}
-	return Iter;
+	return Ret;
 }
 
-static bool PackagesAppearEquivalent( const UPackage* A, const UPackage* B )
+namespace EMergeParticipant
 {
-	TArray<UObject*> SubobjectsA;
-	GetObjectsWithOuter(A, SubobjectsA, true);
-	TArray<UObject*> SubobjectsB;
-	GetObjectsWithOuter(B, SubobjectsB, true);
-
-	auto IsTransientOrPendingKill = [](UObject const& Obj){ return Obj.HasAnyFlags(RF_Transient | RF_PendingKill); };
-	TSet<FString> SubobjectANames;
-	for( const auto& I : SubobjectsA )
-	{
-		if (!FindOuter(*I, IsTransientOrPendingKill ) )
-		{
-			SubobjectANames.Add(I->GetName());
-		}
-	}
-	TSet<FString> SubobjectBNames;
-	for (const auto& I : SubobjectsB)
-	{
-		if (!FindOuter(*I, IsTransientOrPendingKill))
-		{
-			SubobjectBNames.Add(I->GetName());
-		}
-	}
-
-	TArray<UObject*> DifferentObjects;
-	TSet<FString> AB = SubobjectANames.Difference(SubobjectBNames);
-	TSet<FString> BA = SubobjectBNames.Difference(SubobjectANames);
-	TSet<FString> Difference = AB.Union( BA );
-	for(const auto& I : Difference )
+	enum Type
 	{
-		auto FindByName = [&I]( const UObject* Candidate ) { return Candidate->GetName() == I; };
-		UObject** FoundObject = SubobjectsA.FindByPredicate( FindByName );
-		if( !FoundObject )
-		{
-			FoundObject = SubobjectsB.FindByPredicate( FindByName );
-		}
-
-		if( FoundObject )
-		{
-			DifferentObjects.Add( *FoundObject );
-		}
-		else
-		{
-			verify(false);
-		}
-	}
-
-	return Difference.Num() == 0;
+		MERGE_PARTICIPANT_REMOTE,
+		MERGE_PARTICIPANT_BASE,
+		MERGE_PARTICIPANT_LOCAL,
+	};
 }
 
-/** 
- * Commands used to navigate the lists of differences between the local revision and the common base
- * or the remote revision and the common base.
- */
-class FMergeDifferencesListCommands : public TCommands<FMergeDifferencesListCommands>
+struct FBlueprintRevPair
 {
-public:
-	FMergeDifferencesListCommands()
-		: TCommands<FMergeDifferencesListCommands>("MergeList", NSLOCTEXT("Merge", "Blueprint", "Blueprint Merge"), NAME_None, FEditorStyle::GetStyleSetName())
+	FBlueprintRevPair( const UBlueprint* InBlueprint, const FRevisionInfo& InRevData )
+		: Blueprint( InBlueprint )
+		, RevData( InRevData )
 	{
 	}
 
-	/** Go to previous or next difference, focusing on either the remote or local diff */
-	TSharedPtr<FUICommandInfo> LocalPrevious;
-	TSharedPtr<FUICommandInfo> LocalNext;
-	TSharedPtr<FUICommandInfo> RemotePrevious;
-	TSharedPtr<FUICommandInfo> RemoteNext;
-
-	virtual void RegisterCommands() override
-	{
-		UI_COMMAND(LocalPrevious, "LocalPrevious", "Go to previous local difference", EUserInterfaceActionType::Button, FInputGesture(EKeys::F7, EModifierKey::Control));
-		UI_COMMAND(LocalNext, "LocalNext", "Go to next local difference", EUserInterfaceActionType::Button, FInputGesture(EKeys::F7));
-		UI_COMMAND(RemotePrevious, "RemotePrevious", "Go to previous remote difference", EUserInterfaceActionType::Button, FInputGesture(EKeys::F8, EModifierKey::Control));
-		UI_COMMAND(RemoteNext, "RemoteNext", "Go to next remote difference", EUserInterfaceActionType::Button, FInputGesture(EKeys::F8));
-	}
+	const UBlueprint* Blueprint;
+	const FRevisionInfo& RevData;
 };
 
-SBlueprintMerge::SBlueprintMerge()
-	: SBlueprintDiff()
-	, OwningEditor()
+struct FMergeGraphRowEntry
 {
-}
+	FString GraphName;
+
+	// These lists contain shared ptrs because they are displayed by a
+	// SListView, and it currently does not support lists of values.
+	TArray< TSharedPtr<FDiffSingleResult> > RemoteDifferences;
+	TArray< TSharedPtr<FDiffSingleResult> > LocalDifferences;
+	
+	bool bExistsInRemote;
+	bool bExistsInBase;
+	bool bExistsInLocal;
+};
 
-void SBlueprintMerge::Construct(const FArguments& InArgs)
+static UEdGraph* FindGraphByName(UBlueprint const& FromBlueprint, const FString& GraphName)
 {
-	FMergeDifferencesListCommands::Register();
-
-	PanelLocal.Blueprint = InArgs._BlueprintLocal;
-
-	OwningEditor = InArgs._OwningEditor;
-	check( OwningEditor.Pin().IsValid() );
+	TArray<UEdGraph*> Graphs;
+	FromBlueprint.GetAllGraphs(Graphs);
 
-	return SBlueprintDiff::Construct( InArgs._BaseArgs );
+	UEdGraph* Ret = NULL;
+	if (UEdGraph** Result = Graphs.FindByPredicate(FMatchName(GraphName)))
+	{
+		Ret = *Result;
+	}
+	return Ret;
 }
 
-TSharedRef<ITableRow> SBlueprintMerge::OnGenerateRow( TSharedPtr<struct FDiffSingleResult> ParamItem, const TSharedRef<STableViewBase>& OwnerTable)
+static FDiffPanel InitializePanel(const FBlueprintRevPair& BlueprintRevPair)
 {
-	FDiffSingleResult const& Result = *ParamItem;
-	FString ToolTip = Result.ToolTip;
-	FLinearColor Color = Result.DisplayColor;
-	FString Text = Result.DisplayString;
-	if (Text.Len() == 0)
-	{
-		Text = LOCTEXT("DIF_UnknownDiff", "Unknown Diff").ToString();
-		ToolTip = LOCTEXT("DIF_Confused", "There is an unspecified difference").ToString();
-	}
-	return SNew(STableRow< TSharedPtr<struct FDiffSingleResult> >, OwnerTable)
-		.Content()
+	FDiffPanel Ret;
+	Ret.Blueprint = BlueprintRevPair.Blueprint;
+	SAssignNew(Ret.GraphEditorBorder, SBorder)
+		.VAlign(VAlign_Fill)
 		[
-			SNew(STextBlock)
-			.ToolTipText(ToolTip)
-			.ColorAndOpacity(Color)
-			.Text(Text)
+			SBlueprintDiff::DefaultEmptyPanel()
 		];
+	Ret.RevisionInfo = BlueprintRevPair.RevData;
+	Ret.bShowAssetName = false;
+	return Ret;
 }
 
-void SBlueprintMerge::OnDiffListSelectionChanged(TSharedPtr<struct FDiffSingleResult> Item, ESelectInfo::Type SelectionType)
+// @todo doc: can i avoid shared ptrs for individual entries?
+static TArray< TSharedPtr<FMergeGraphRowEntry> > GenerateDiffListItems( const FBlueprintRevPair& RemoteBlueprint, const FBlueprintRevPair& BaseBlueprint, const FBlueprintRevPair& LocalBlueprint)
 {
-	if( !Item.IsValid() )
-	{
-		return;
+	// Index all the graphs by name, we use the name of the graph as the 
+	// basis of comparison between the various versions of the blueprint.
+	TMap< FString, UEdGraph* > RemoteGraphMap, BaseGraphMap, LocalGraphMap;
+	// We also want the set of all graph names in these blueprints, so that we 
+	// can iterate over every graph.
+	TSet< FString > AllGraphNames;
+	{ 
+		TArray<UEdGraph*> GraphsRemote, GraphsBase, GraphsLocal;
+		RemoteBlueprint.Blueprint->GetAllGraphs(GraphsRemote);
+		BaseBlueprint.Blueprint->GetAllGraphs(GraphsBase);
+		LocalBlueprint.Blueprint->GetAllGraphs(GraphsLocal);
+
+		const auto ToMap = [&AllGraphNames]( const TArray<UEdGraph*>& InList, TMap<FString, UEdGraph*>& OutMap)
+		{
+			for( auto Graph : InList )
+			{
+				OutMap.Add( Graph->GetName(), Graph );
+				AllGraphNames.Add(Graph->GetName());
+			}
+		};
+		ToMap(GraphsRemote, RemoteGraphMap);
+		ToMap(GraphsBase, BaseGraphMap);
+		ToMap(GraphsLocal, LocalGraphMap);
 	}
-
-	//focus the graph onto the diff that was clicked on
-	FDiffSingleResult const& Result = *Item;
-	if (Result.Pin1)
+	
+	TArray< TSharedPtr<FMergeGraphRowEntry> > Ret;
 	{
-		PanelNew.GraphEditor.Pin()->ClearSelectionSet();
-		PanelOld.GraphEditor.Pin()->ClearSelectionSet();
-
-		if (Result.Pin1)
+		const auto GenerateDifferences = []( UEdGraph* GraphNew, UEdGraph** GraphOld )
 		{
-			Result.Pin1->bIsDiffing = true;
-			GetGraphEditorForGraph(Result.Pin1->GetOwningNode()->GetGraph())->JumpToPin(Result.Pin1);
-		}
+			TArray<FDiffSingleResult> Results;
+			FGraphDiffControl::DiffGraphs(GraphOld ? *GraphOld : NULL, GraphNew, Results );
+			return Results;
+		};
 
-		if (Result.Pin2)
+		for( const auto& GraphName : AllGraphNames )
 		{
-			Result.Pin2->bIsDiffing = true;
-			GetGraphEditorForGraph(Result.Pin2->GetOwningNode()->GetGraph())->JumpToPin(Result.Pin2);
+			TArray< FDiffSingleResult > RemoteDifferences;
+			TArray< FDiffSingleResult > LocalDifferences;
+			bool bExistsInRemote, bExistsInBase, bExistsInLocal;
+
+			{
+				UEdGraph** RemoteGraph = RemoteGraphMap.Find(GraphName);
+				UEdGraph** BaseGraph = BaseGraphMap.Find(GraphName);
+				UEdGraph** LocalGraph = LocalGraphMap.Find(GraphName);
+
+				if( RemoteGraph )
+				{
+					RemoteDifferences = GenerateDifferences(*RemoteGraph, BaseGraph);
+				}
+
+				if( LocalGraph )
+				{
+					LocalDifferences = GenerateDifferences(*LocalGraph, BaseGraph);
+				}
+				
+				bExistsInRemote = RemoteGraph != NULL;
+				bExistsInBase = BaseGraph != NULL;
+				bExistsInLocal = LocalGraph != NULL;
+			}
+
+			FMergeGraphRowEntry NewEntry = {
+				GraphName,
+				HackToShared(RemoteDifferences),
+				HackToShared(LocalDifferences),
+				bExistsInRemote,
+				bExistsInBase,
+				bExistsInLocal
+			};
+			Ret.Add( TSharedPtr<FMergeGraphRowEntry>( new FMergeGraphRowEntry(NewEntry) ) );
 		}
 	}
-	else if (Result.Node1)
-	{
-		PanelNew.GraphEditor.Pin()->ClearSelectionSet();
-		PanelOld.GraphEditor.Pin()->ClearSelectionSet();
 
-		if (Result.Node2)
+	return Ret;
+}
+
+static void LockViews( TArray<FDiffPanel>& Views, bool bAreLocked )
+{
+	for( auto& Panel : Views )
+	{
+		auto GraphEditor = Panel.GraphEditor.Pin();
+		if( GraphEditor.IsValid() )
 		{
-			GetGraphEditorForGraph(Result.Node2->GetGraph())->JumpToNode(Result.Node2, false);
+			// lock this panel to ever other panel:
+			for (auto& OtherPanel : Views)
+			{
+				auto OtherGraphEditor = OtherPanel.GraphEditor.Pin();
+				if( OtherGraphEditor.IsValid() &&
+					OtherGraphEditor != GraphEditor )
+				{
+					if( bAreLocked )
+					{
+						GraphEditor->LockToGraphEditor(OtherGraphEditor);
+					}
+					else
+					{
+						GraphEditor->UnlockFromGraphEditor(OtherGraphEditor);
+					}
+				}
+			}
 		}
-		GetGraphEditorForGraph(Result.Node1->GetGraph())->JumpToNode(Result.Node1, false);
 	}
 }
 
-void SBlueprintMerge::ResetGraphEditors()
+FDiffPanel& GetDiffPanelForNode( const UEdGraphNode& Node, TArray< FDiffPanel >& Panels )
 {
-	auto GraphEditor = PanelLocal.GraphEditor.Pin();
-	if (GraphEditor != TSharedPtr<SGraphEditor>())
+	for( auto& Panel : Panels )
 	{
-		if (bLockViews)
+		auto GraphEditor = Panel.GraphEditor.Pin();
+		if( GraphEditor.IsValid() )
 		{
-			GraphEditor->LockToGraphEditor(PanelNew.GraphEditor);
-			GraphEditor->LockToGraphEditor(PanelOld.GraphEditor);
-
-			auto OldGraph = PanelOld.GraphEditor.Pin();
-			if( OldGraph != TSharedPtr<SGraphEditor>() )
-			{
-				OldGraph->LockToGraphEditor( GraphEditor );
-			}
-			auto NewGraph = PanelNew.GraphEditor.Pin();
-			if( NewGraph!= TSharedPtr<SGraphEditor>() )
+			if( Node.GetGraph() == GraphEditor->GetCurrentGraph() )
 			{
-				NewGraph->LockToGraphEditor(GraphEditor);
+				return Panel;
 			}
 		}
-		else
-		{
-			GraphEditor->UnlockFromGraphEditor(PanelNew.GraphEditor);
-			GraphEditor->UnlockFromGraphEditor(PanelOld.GraphEditor);
+	}
+	checkf(false, TEXT("Looking for node %s but it cannot be found in provided panels"), *Node.GetName() );
+	static FDiffPanel Default;
+	return Default;
+}
+
+class SMergeGraphView : public SCompoundWidget
+{
+public:
+	SLATE_BEGIN_ARGS(SMergeGraphView)
+	{}
+	SLATE_END_ARGS()
+
+	void Construct(const FArguments InArgs, const FBlueprintMergeData& InData);
+
+private:
+	void OnGraphListSelectionChanged(TSharedPtr<FMergeGraphRowEntry> Item, ESelectInfo::Type SelectionType);
+	void OnDiffListSelectionChanged(TSharedPtr<FDiffSingleResult> Item, ESelectInfo::Type SelectionType);
 
-			auto OldGraph = PanelOld.GraphEditor.Pin();
-			if (OldGraph != TSharedPtr<SGraphEditor>())
+	FDiffPanel& GetRemotePanel() { return DiffPanels[EMergeParticipant::MERGE_PARTICIPANT_REMOTE]; }
+	FDiffPanel& GetBasePanel() { return DiffPanels[EMergeParticipant::MERGE_PARTICIPANT_BASE]; }
+	FDiffPanel& GetLocalPanel() { return DiffPanels[EMergeParticipant::MERGE_PARTICIPANT_LOCAL]; }
+
+	FReply OnToggleLockView();
+	const FSlateBrush*  GetLockViewImage() const;
+
+	TArray< FDiffPanel > DiffPanels;
+	FBlueprintMergeData Data;
+
+	TArray< TSharedPtr<FMergeGraphRowEntry> > DifferencesFromBase;
+	TSharedPtr<SBorder> DiffResultsWidget;
+
+	bool bViewsAreLocked;
+};
+
+// Some color constants used by the merge view:
+const FLinearColor MergeBlue = FLinearColor(0.3f, 0.3f, 1.f);
+const FLinearColor MergeRed = FLinearColor(1.0f, 0.2f, 0.3f);
+const FLinearColor MergeYellow = FLinearColor::Yellow;
+const FLinearColor MergeWhite = FLinearColor::White;
+
+void SMergeGraphView::Construct( const FArguments InArgs, const FBlueprintMergeData& InData )
+{
+	Data = InData;
+	bViewsAreLocked = true;
+
+	TArray<FBlueprintRevPair> BlueprintsForDisplay;
+	// EMergeParticipant::MERGE_PARTICIPANT_REMOTE
+	BlueprintsForDisplay.Add(FBlueprintRevPair(InData.BlueprintNew, InData.RevisionNew));
+	// EMergeParticipant::MERGE_PARTICIPANT_BASE
+	BlueprintsForDisplay.Add(FBlueprintRevPair(InData.BlueprintBase, InData.RevisionBase));
+	// EMergeParticipant::MERGE_PARTICIPANT_LOCAL
+	BlueprintsForDisplay.Add( FBlueprintRevPair(InData.BlueprintLocal, FRevisionInfo()) );
+
+	for( const auto& BlueprintRevPair : BlueprintsForDisplay )
+	{
+		DiffPanels.Add( InitializePanel( BlueprintRevPair ) );
+	}
+
+	TSharedRef<SSplitter> PanelContainer = SNew(SSplitter);
+	for( const auto& Panel : DiffPanels )
+	{
+		PanelContainer->AddSlot()[ Panel.GraphEditorBorder.ToSharedRef() ];
+	}
+
+	DifferencesFromBase = GenerateDiffListItems(	BlueprintsForDisplay[ EMergeParticipant::MERGE_PARTICIPANT_REMOTE ]
+													,BlueprintsForDisplay[ EMergeParticipant::MERGE_PARTICIPANT_BASE ]
+													,BlueprintsForDisplay[ EMergeParticipant::MERGE_PARTICIPANT_LOCAL ]);
+
+	// This is the function we'll use to generate a row in the control that lists all the available graphs:
+	const auto RowGenerator = [](TSharedPtr<FMergeGraphRowEntry> ParamItem, const TSharedRef<STableViewBase>& OwnerTable) -> TSharedRef<ITableRow>
+	{
+		// blue indicates added, red indicates changed, yellow indicates removed, white indicates no change:
+		const auto ComputeColorForRevision = [](bool bExistsInBase, bool bExistsInThisRevision, bool bHasChangesInThisRevision) -> FLinearColor
+		{
+			if( bExistsInBase )
 			{
-				OldGraph->UnlockFromGraphEditor(GraphEditor);
+				if( !bExistsInThisRevision )
+				{
+					return MergeYellow;
+				}
+				else if( bHasChangesInThisRevision )
+				{
+					return MergeRed;
+				}
 			}
-			auto NewGraph = PanelNew.GraphEditor.Pin();
-			if (NewGraph != TSharedPtr<SGraphEditor>())
+			if( !bExistsInBase )
 			{
-				NewGraph->UnlockFromGraphEditor(GraphEditor);
+				if( bExistsInThisRevision  )
+				{
+					return MergeBlue;
+				}
 			}
-		}
-	}
-	return SBlueprintDiff::ResetGraphEditors();
-}
 
-TSharedRef<SWidget> SBlueprintMerge::GenerateDiffWindow()
-{
-	return SNew(SVerticalBox)
-		+ SVerticalBox::Slot()
+			return MergeWhite;
+		};
+
+		FLinearColor RemoteColor = ComputeColorForRevision(ParamItem->bExistsInBase, ParamItem->bExistsInRemote, ParamItem->RemoteDifferences.Num() != 0);
+		FLinearColor BaseColor = ComputeColorForRevision(ParamItem->bExistsInBase, ParamItem->bExistsInBase, false);
+		FLinearColor LocalColor = ComputeColorForRevision(ParamItem->bExistsInBase, ParamItem->bExistsInLocal, ParamItem->LocalDifferences.Num() != 0);
+
+		const auto Box = [] (bool bIsPresent, FLinearColor Color) -> SHorizontalBox::FSlot&
+		{
+			return SHorizontalBox::Slot()
+				.AutoWidth()
+				.HAlign(HAlign_Right)
+				.VAlign(VAlign_Center)
+				.MaxWidth(8.0f)
+				[
+					SNew(SImage)
+					.ColorAndOpacity(Color)
+					.Image(bIsPresent ? FEditorStyle::GetBrush("BlueprintDif.HasGraph") : FEditorStyle::GetBrush("BlueprintDif.MissingGraph"))
+				];
+		};
+
+		return SNew(STableRow< TSharedPtr<FMergeGraphRowEntry> >, OwnerTable)
+			.ToolTipText(LOCTEXT("MergeGraphsDifferentText", "@todo doc"))
+			.Content()
+			[
+				SNew(SHorizontalBox)
+					+ SHorizontalBox::Slot()
+					[
+						SNew(STextBlock)
+						.ColorAndOpacity(FLinearColor::White)
+						.Text(ParamItem->GraphName)
+					]
+					+ Box(ParamItem->bExistsInRemote, RemoteColor )
+					+ Box(ParamItem->bExistsInBase, BaseColor)
+					+ Box(ParamItem->bExistsInLocal, LocalColor)
+			];
+	};
+
+	ChildSlot
+	[
+		SNew(SSplitter)
+		.Orientation(Orient_Horizontal)
+		+ SSplitter::Slot()
+		.Value(0.1f)
 		[
-			SNew(SSplitter).Orientation(Orient_Vertical)
-			+ SSplitter::Slot()
-			.Value(.7f)
+			// ToolPanel.GroupBorder is currently what gives this control its dark grey appearance:
+			SNew(SBorder)
+			.BorderImage(FEditorStyle::GetBrush("NoBorder"))
+			.Padding(FMargin(1.0f))
 			[
 				SNew(SSplitter)
-				+ SSplitter::Slot()
-				.Value(0.5f)
+				.Orientation(Orient_Vertical)
+				+SSplitter::Slot()
+				.Value(.3f)
 				[
-					//left blueprint
-					SAssignNew(PanelLocal.GraphEditorBorder, SBorder)
-					.VAlign(VAlign_Fill)
+					SNew(SBorder)
+					.BorderImage(FEditorStyle::GetBrush("ToolPanel.GroupBorder"))
 					[
-						SBlueprintDiff::DefaultEmptyPanel()
+						SNew(SVerticalBox)
+						+ SVerticalBox::Slot()
+						.AutoHeight()
+						.HAlign(HAlign_Left)
+						[
+							SNew(SButton)
+							.OnClicked(this, &SMergeGraphView::OnToggleLockView)
+							.Content()
+							[
+								SNew(SImage)
+								.Image(this, &SMergeGraphView::GetLockViewImage)
+								.ToolTipText(LOCTEXT("BPDifLock", "Lock the blueprint views?"))
+							]
+						]
+						+ SVerticalBox::Slot()
+						[
+							SNew(SListView< TSharedPtr<FMergeGraphRowEntry> >)
+							.ItemHeight(24)
+							.ListItemsSource(&DifferencesFromBase)
+							.OnGenerateRow(SListView< TSharedPtr<FMergeGraphRowEntry> >::FOnGenerateRow::CreateStatic(RowGenerator))
+							.SelectionMode(ESelectionMode::Single)
+							.OnSelectionChanged(this, &SMergeGraphView::OnGraphListSelectionChanged)
+						]
 					]
 				]
 				+ SSplitter::Slot()
-				.Value(0.5f)
+				.Value(.7f)
 				[
-					//middle (base) blueprint
-					SAssignNew(PanelOld.GraphEditorBorder, SBorder)
-					.VAlign(VAlign_Fill)
-					[
-						SBlueprintDiff::DefaultEmptyPanel()
-					]
-				]
-			+ SSplitter::Slot()
-				.Value(0.5f)
-				[
-					//right blueprint
-					SAssignNew(PanelNew.GraphEditorBorder, SBorder)
-					.VAlign(VAlign_Fill)
-					[
-						SBlueprintDiff::DefaultEmptyPanel()
-					]
+					SAssignNew(DiffResultsWidget, SBorder)
+					.BorderImage(FEditorStyle::GetBrush("ToolPanel.GroupBorder"))
 				]
 			]
 		]
-	;
-}
-
-TSharedRef<SWidget> SBlueprintMerge::GenerateToolbar()
-{
-	return SNew(SVerticalBox)
-		+ SVerticalBox::Slot()
-		.AutoHeight()
+		+ SSplitter::Slot()
+		.Value(0.9f)
 		[
-			SNew(SVerticalBox)
-			+ SVerticalBox::Slot()
-			[
-				SNew(SButton)
-				.OnClicked(this, &SBlueprintMerge::OnAcceptResultClicked)
-				.Content()
-				[
-					SNew(STextBlock)
-					.Text(LOCTEXT("AcceptResult", "Accept Result"))
-					.ToolTipText( LOCTEXT( "AcceptResultToolTip", "Overwrite the local blueprint with the blueprint constructed in the lower portion of the merge tool." ) )
-				]
-			]
-			+ SVerticalBox::Slot()
-			[
-				SNew(SButton)
-				.OnClicked(this, &SBlueprintMerge::OnCancelClicked)
-				.Content()
-				[
-					SNew(STextBlock)
-					.Text(LOCTEXT("CancelMerge", "Cancel"))
-					.ToolTipText(LOCTEXT("CancelMergeToolTip", "Discards the changes made in the merge tool. Local file will remain unresolved so the merge tool can be run again if desired."))
-				]
-			]
+			PanelContainer
 		]
-		+ SVerticalBox::Slot()
-			.AutoHeight()
-		[
-			SAssignNew(DiffListBorder, SBorder)
-		];
+	];
 }
 
-static UEdGraph* FindGraphByName(UBlueprint const& FromBlueprint, const FString& GraphName)
+void SMergeGraphView::OnGraphListSelectionChanged(TSharedPtr<FMergeGraphRowEntry> Item, ESelectInfo::Type SelectionType)
 {
-	TArray<UEdGraph*> Graphs;
-	FromBlueprint.GetAllGraphs(Graphs);
-
-	UEdGraph* Ret = NULL;
-	if (UEdGraph** Result = Graphs.FindByPredicate(FMatchName(GraphName)))
+	if (SelectionType != ESelectInfo::OnMouseClick || !Item.IsValid())
 	{
-		Ret = *Result;
+		return;
 	}
-	return Ret;
-}
 
-void SBlueprintMerge::HandleGraphChanged(const FString& GraphName)
-{
-	UEdGraph* GraphOld = FindGraphByName(*PanelOld.Blueprint, GraphName);
-	UEdGraph* GraphNew = FindGraphByName(*PanelNew.Blueprint, GraphName);
-	UEdGraph* GraphLocal = FindGraphByName(*PanelLocal.Blueprint, GraphName);
+	FString GraphName = Item->GraphName;
+
+	UEdGraph* GraphRemote = FindGraphByName(*GetRemotePanel().Blueprint, GraphName);
+	UEdGraph* GraphBase = FindGraphByName(*GetBasePanel().Blueprint, GraphName);
+	UEdGraph* GraphLocal = FindGraphByName(*GetLocalPanel().Blueprint, GraphName);
+
+	GetBasePanel().GeneratePanel(GraphBase, NULL);
+	GetRemotePanel().GeneratePanel(GraphRemote, GraphBase);
+	GetLocalPanel().GeneratePanel(GraphLocal, GraphBase);
+
+	LockViews(DiffPanels, bViewsAreLocked);
 
-	PanelOld.GeneratePanel(GraphOld, NULL);
-	PanelNew.GeneratePanel(GraphNew, GraphOld);
-	PanelLocal.GeneratePanel(GraphLocal, GraphOld);
+	DiffResultsWidget->SetContent(
+		SNew( SListView< TSharedPtr<FDiffSingleResult> >)
+		.ItemHeight(24)
+		.ListItemsSource(&Item->LocalDifferences)
+		.OnGenerateRow( SListView< TSharedPtr<FDiffSingleResult> >::FOnGenerateRow::CreateStatic(
+			[](TSharedPtr<FDiffSingleResult> ParamItem, const TSharedRef<STableViewBase>& OwnerTable) -> TSharedRef<ITableRow>
+			{
+				FDiffResultItem WidgetGenerator( *ParamItem );
+				return SNew(STableRow< TSharedPtr<FDiffSingleResult> >, OwnerTable)
+					.Content()
+					[
+						WidgetGenerator.GenerateWidget()
+					];
+			}
+		) )
+		.SelectionMode(ESelectionMode::Single)
+		.OnSelectionChanged(this, &SMergeGraphView::OnDiffListSelectionChanged)
+	);
 }
 
-template< typename T >
-static void CopyToShared( const TArray<T>& Source, TArray< TSharedPtr<T> >& Dest )
+void SMergeGraphView::OnDiffListSelectionChanged(TSharedPtr<FDiffSingleResult> Item, ESelectInfo::Type SelectionType)
 {
-	check( Dest.Num() == 0 );
-	for( auto i : Source )
+	if( SelectionType != ESelectInfo::OnMouseClick || !Item.IsValid() )
 	{
-		Dest.Add( MakeShareable(new T( i )) );
+		return;
 	}
-}
 
-static void WriteBackup( UPackage& Package, const FString& Directory, const FString& Filename )
-{
-	if( GEditor )
+	// Clear graph selection:
+	for( auto& Panel : DiffPanels )
 	{
-		const FString DestinationFilename = Directory + FString("/") + Filename;
-		FString OriginalFilename;
-		if( FPackageName::DoesPackageExist(Package.GetName(), NULL, &OriginalFilename) )
+		auto GraphEdtiorPtr = Panel.GraphEditor.Pin();
+		if (GraphEdtiorPtr.IsValid())
 		{
-			IFileManager::Get().Copy(*DestinationFilename, *OriginalFilename);
+			GraphEdtiorPtr->ClearSelectionSet();
+		}
+	}
+
+	//focus the graph onto the diff that was clicked on
+	if (Item->Pin1)
+	{
+		GetDiffPanelForNode(*Item->Pin1->GetOwningNode(), DiffPanels).FocusDiff( *Item->Pin1 );
+
+		if (Item->Pin2)
+		{
+			GetDiffPanelForNode(*Item->Pin2->GetOwningNode(), DiffPanels).FocusDiff( *Item->Pin2 );
+		}
+	}
+	else if (Item->Node1)
+	{
+		GetDiffPanelForNode(*Item->Node1, DiffPanels).FocusDiff( *Item->Node1 );
+
+		if (Item->Node2)
+		{
+			GetDiffPanelForNode(*Item->Node2, DiffPanels).FocusDiff( *Item->Node2 );
 		}
 	}
 }
 
-void SBlueprintMerge::OnSelectionChanged(SBlueprintDiff::FGraphToDiff Item, ESelectInfo::Type SelectionType)
+FReply SMergeGraphView::OnToggleLockView()
 {
-	if (SelectionType != ESelectInfo::OnMouseClick || !Item.IsValid())
+	bViewsAreLocked = !bViewsAreLocked;
+
+	LockViews( DiffPanels, bViewsAreLocked );
+	return FReply::Handled();
+}
+
+const FSlateBrush* SMergeGraphView::GetLockViewImage() const
+{
+	return bViewsAreLocked ? FEditorStyle::GetBrush("GenericLock") : FEditorStyle::GetBrush("GenericUnlock");
+}
+
+/** 
+ * Commands used to navigate the lists of differences between the local revision and the common base
+ * or the remote revision and the common base.
+ */
+class FMergeDifferencesListCommands : public TCommands<FMergeDifferencesListCommands>
+{
+public:
+	FMergeDifferencesListCommands()
+		: TCommands<FMergeDifferencesListCommands>("MergeList", NSLOCTEXT("Merge", "Blueprint", "Blueprint Merge"), NAME_None, FEditorStyle::GetStyleSetName())
 	{
-		return;
 	}
 
-	FString GraphName = Item->GetGraphOld()->GetName();
-	DisablePinDiffFocus();
-	HandleGraphChanged(GraphName);
+	/** Go to previous or next difference, focusing on either the remote or local diff */
+	TSharedPtr<FUICommandInfo> LocalPrevious;
+	TSharedPtr<FUICommandInfo> LocalNext;
+	TSharedPtr<FUICommandInfo> RemotePrevious;
+	TSharedPtr<FUICommandInfo> RemoteNext;
+
+	virtual void RegisterCommands() override
+	{
+		UI_COMMAND(LocalPrevious, "LocalPrevious", "Go to previous local difference", EUserInterfaceActionType::Button, FInputGesture(EKeys::F7, EModifierKey::Control));
+		UI_COMMAND(LocalNext, "LocalNext", "Go to next local difference", EUserInterfaceActionType::Button, FInputGesture(EKeys::F7));
+		UI_COMMAND(RemotePrevious, "RemotePrevious", "Go to previous remote difference", EUserInterfaceActionType::Button, FInputGesture(EKeys::F8, EModifierKey::Control));
+		UI_COMMAND(RemoteNext, "RemoteNext", "Go to next remote difference", EUserInterfaceActionType::Button, FInputGesture(EKeys::F8));
+	}
+};
 
-	// Ensure that regenerated graphs are properly locked/unlocked:
-	ResetGraphEditors();
+SBlueprintMerge::SBlueprintMerge()
+	: Data()
+	, GraphView()
+	, ComponentsView()
+	, DefaultsView()
+{
+}
 
-	// We'll diff the other two graphs against this shared base:
-	UEdGraph* GraphBase = FindGraphByName(*PanelOld.Blueprint, GraphName);
+void SBlueprintMerge::Construct(const FArguments InArgs, const FBlueprintMergeData& InData)
+{
+	FMergeDifferencesListCommands::Register();
 
-	// Diff local version vs base:
-	TArray<FDiffSingleResult> FoundDiffsLocal;
-	UEdGraph* GraphLocal = FindGraphByName(*PanelLocal.Blueprint, GraphName);
-	FGraphDiffControl::DiffGraphs(GraphLocal, GraphBase, FoundDiffsLocal);
+	check( InData.OwningEditor.Pin().IsValid() );
 
-	// Diff remote version vs. base:
-	TArray<FDiffSingleResult> FoundDiffsRemote;
-	UEdGraph* GraphRemote = FindGraphByName(*PanelNew.Blueprint, GraphName);
-	FGraphDiffControl::DiffGraphs(GraphRemote, GraphBase, FoundDiffsRemote);
+	Data = InData;
 
-	// Display results of the two diffs:
-	FMergeDifferencesListCommands const& Commands = FMergeDifferencesListCommands::Get();
-	LocalDiffResults.Empty();
-	TSharedRef< SWidget > DiffsLocal = GenerateDiffView(FoundDiffsLocal, Commands.LocalNext, Commands.LocalPrevious, LocalDiffResults );
-	RemoteDiffResults.Empty();
-	TSharedRef< SWidget > DiffsRemote = GenerateDiffView(FoundDiffsRemote, Commands.RemoteNext, Commands.RemotePrevious, RemoteDiffResults );
+	GraphView = SNew( SMergeGraphView, InData );
 
-	DiffListBorder->SetContent( SNew( SVerticalBox )
-		+ SVerticalBox::Slot()[DiffsLocal]
-		+ SVerticalBox::Slot()[DiffsRemote] );
+	this->ChildSlot
+		[
+			GraphView.ToSharedRef()
+		];
+
+	/*ComponentsView;
+	DefaultsView;*/
 }
 
-SGraphEditor* SBlueprintMerge::GetGraphEditorForGraph(UEdGraph* Graph) const
+template< typename T >
+static void CopyToShared( const TArray<T>& Source, TArray< TSharedPtr<T> >& Dest )
 {
-	if (PanelLocal.GraphEditor.Pin()->GetCurrentGraph() == Graph)
+	check( Dest.Num() == 0 );
+	for( auto i : Source )
 	{
-		return PanelLocal.GraphEditor.Pin().Get();
+		Dest.Add( MakeShareable(new T( i )) );
 	}
-	return SBlueprintDiff::GetGraphEditorForGraph(Graph);
 }
 
-TSharedRef< SWidget > SBlueprintMerge::GenerateDiffView(TArray<FDiffSingleResult>& Diffs, TSharedPtr< const FUICommandInfo > CommandNext, TSharedPtr< const FUICommandInfo > CommandPrev, TArray< TSharedPtr<FDiffSingleResult> >& SharedResults )
+static void WriteBackup( UPackage& Package, const FString& Directory, const FString& Filename )
 {
-	struct SortDiff
+	if( GEditor )
 	{
-		bool operator () (const FDiffSingleResult& A, const FDiffSingleResult& B) const
+		const FString DestinationFilename = Directory + FString("/") + Filename;
+		FString OriginalFilename;
+		if( FPackageName::DoesPackageExist(Package.GetName(), NULL, &OriginalFilename) )
 		{
-			return A.Diff < B.Diff;
+			IFileManager::Get().Copy(*DestinationFilename, *OriginalFilename);
 		}
-	};
-
-	Sort(Diffs.GetTypedData(), Diffs.Num(), SortDiff());
-
-	TSharedPtr<FUICommandList> KeyCommands = MakeShareable(new FUICommandList);
-
-	//KeyCommands->MapAction(CommandPrev, FExecuteAction::CreateSP(this, &FListItemGraphToDiff::PrevDiff));
-	//KeyCommands->MapAction(CommandNext, FExecuteAction::CreateSP(this, &FListItemGraphToDiff::NextDiff));
-
-	FToolBarBuilder ToolbarBuilder(KeyCommands, FMultiBoxCustomization::None);
-	ToolbarBuilder.AddToolBarButton(CommandPrev, NAME_None, TAttribute<FText>(), TAttribute<FText>(), FSlateIcon(FEditorStyle::GetStyleSetName(), "BlueprintDif.PrevDiff"));
-	ToolbarBuilder.AddToolBarButton(CommandNext, NAME_None, TAttribute<FText>(), TAttribute<FText>(), FSlateIcon(FEditorStyle::GetStyleSetName(), "BlueprintDif.NextDiff"));
-
-	CopyToShared(Diffs, SharedResults);
-
-	TSharedRef<SHorizontalBox> Result = SNew(SHorizontalBox)
-		+ SHorizontalBox::Slot()
-		.FillWidth(1.f)
-		.MaxWidth(350.f)
-		[
-			SNew(SVerticalBox)
-			+ SVerticalBox::Slot()
-			.Padding(0.f)
-			.AutoHeight()
-			[
-				ToolbarBuilder.MakeWidget()
-			]
-			+ SVerticalBox::Slot()
-				.Padding(0.f)
-				.AutoHeight()
-				[
-					SNew(SBorder)
-					.BorderImage(FEditorStyle::GetBrush("PropertyWindow.CategoryBackground"))
-					.Padding(FMargin(2.0f))
-					.ForegroundColor(FEditorStyle::GetColor("PropertyWindow.CategoryForeground"))
-					.ToolTipText(LOCTEXT("BlueprintDifDifferencesToolTip", "List of differences found between revisions, click to select"))
-					.HAlign(HAlign_Center)
-					[
-						SNew(STextBlock)
-						.Text(LOCTEXT("RevisionDifferences", "Revision Differences"))
-					]
-				]
-			+ SVerticalBox::Slot()
-				.Padding(1.f)
-				.FillHeight(1.f)
-				[
-					SNew(SListView< TSharedPtr<FDiffSingleResult> >)
-					.ItemHeight(24)
-					.ListItemsSource(&SharedResults)
-					.OnGenerateRow(this, &SBlueprintMerge::OnGenerateRow)
-					.SelectionMode(ESelectionMode::Single)
-					.OnSelectionChanged( this, &SBlueprintMerge::OnDiffListSelectionChanged)
-				]
-		];
-	return Result;
+	}
 }
 
 UBlueprint* SBlueprintMerge::GetTargetBlueprint()
 {
-	// @todo doc, do we need to cache this on merge construction?
-	return OwningEditor.Pin()->GetBlueprintObj();
+	return Data.OwningEditor.Pin()->GetBlueprintObj();
 }
 
 FReply SBlueprintMerge::OnAcceptResultClicked()
 {
 	// Because merge operations are so destructive and can be confusing, lets write backups of the files involved:
 	const FString BackupSubDir = FPaths::GameSavedDir() / TEXT("Backup") / TEXT("Resolve_Backup[") + FDateTime::Now().ToString(TEXT("%Y-%m-%d-%H-%M-%S")) + TEXT("]");
-	WriteBackup(*PanelNew.Blueprint->GetOutermost(), BackupSubDir, TEXT("RemoteAsset") + FPackageName::GetAssetPackageExtension() );
-	WriteBackup(*PanelOld.Blueprint->GetOutermost(), BackupSubDir, TEXT("CommonBaseAsset") + FPackageName::GetAssetPackageExtension() );
-	WriteBackup(*PanelLocal.Blueprint->GetOutermost(), BackupSubDir, TEXT("LocalAsset") + FPackageName::GetAssetPackageExtension());
+	WriteBackup(*Data.BlueprintNew->GetOutermost(), BackupSubDir, TEXT("RemoteAsset") + FPackageName::GetAssetPackageExtension() );
+	WriteBackup(*Data.BlueprintBase->GetOutermost(), BackupSubDir, TEXT("CommonBaseAsset") + FPackageName::GetAssetPackageExtension() );
+	WriteBackup(*Data.BlueprintLocal->GetOutermost(), BackupSubDir, TEXT("LocalAsset") + FPackageName::GetAssetPackageExtension());
 
 	UPackage* Package = GetTargetBlueprint()->GetOutermost();
 	TArray<UPackage*> PackagesToSave;
@@ -503,14 +595,14 @@ FReply SBlueprintMerge::OnAcceptResultClicked()
 		FSlateNotificationManager::Get().AddNotification( ErrorNotification );
 	}
 
-	OwningEditor.Pin()->CloseMergeTool();
+	Data.OwningEditor.Pin()->CloseMergeTool();
 
 	return FReply::Handled();
 }
 
 FReply SBlueprintMerge::OnCancelClicked()
 {
-	OwningEditor.Pin()->CloseMergeTool();
+	Data.OwningEditor.Pin()->CloseMergeTool();
 	return FReply::Handled();
 }
 
diff --git a/Engine/Source/Developer/Merge/Private/SBlueprintMerge.h b/Engine/Source/Developer/Merge/Private/SBlueprintMerge.h
--- a/Engine/Source/Developer/Merge/Private/SBlueprintMerge.h
+++ b/Engine/Source/Developer/Merge/Private/SBlueprintMerge.h
@@ -1,52 +1,38 @@
 // Copyright 1998-2014 Epic Games, Inc. All Rights Reserved.
 #pragma once
 
-#include "SBlueprintDiff.h"
+#include "BlueprintMergeData.h"
 
 class FBlueprintEditor;
 
-class MERGE_API SBlueprintMerge : public SBlueprintDiff
+class MERGE_API SBlueprintMerge : public SCompoundWidget
 {
 public:
 	SBlueprintMerge();
 
-	SLATE_BEGIN_ARGS(SBlueprintMerge){}
-	SLATE_ARGUMENT( const class UBlueprint*, BlueprintLocal)
-	SLATE_ARGUMENT( SBlueprintDiff::FArguments, BaseArgs )
-	SLATE_ARGUMENT( TSharedPtr<FBlueprintEditor>, OwningEditor )
+	SLATE_BEGIN_ARGS(SBlueprintMerge)
+	{}
 	SLATE_END_ARGS()
-	
-	void Construct(const FArguments& InArgs);
-	TSharedRef<ITableRow> OnGenerateRow( TSharedPtr<struct FDiffSingleResult> ParamItem, const TSharedRef<STableViewBase>& OwnerTable);
-	void OnDiffListSelectionChanged(TSharedPtr<struct FDiffSingleResult> Item, ESelectInfo::Type SelectionType);
-
-protected:
-	/** SBlueprintDiff overrides, see base class for documentation */
-	virtual void ResetGraphEditors() override;
-	virtual TSharedRef<SWidget> GenerateDiffWindow() override;
-	virtual TSharedRef<SWidget> GenerateToolbar() override;
-	virtual void HandleGraphChanged(const FString& GraphName) override;
-	virtual void OnSelectionChanged(FGraphToDiff Item, ESelectInfo::Type SelectionType) override;
-	virtual SGraphEditor* GetGraphEditorForGraph(UEdGraph* Graph) const override;
-	
+
+	void Construct(const FArguments InArgs, const FBlueprintMergeData& InData);
+
+protected:	
 	/** Helper functions */
-	TSharedRef< SWidget > GenerateDiffView(TArray<FDiffSingleResult>& Diffs, TSharedPtr< const FUICommandInfo > CommandNext, TSharedPtr< const FUICommandInfo > CommandPrev, TArray< TSharedPtr<FDiffSingleResult> >& SharedResults);
 	UBlueprint* GetTargetBlueprint();
 
 	/** Event handlers */
 	FReply OnAcceptResultClicked();
 	FReply OnCancelClicked();
 
-	// We can't maintain a SharedPtr to OwningEditor because we don't logically own it,
-	// if you make this a SharedPtr closing the blueprint editor while the merge tool
-	// is open will prevent the blueprint editor destructor from running:
-	TWeakPtr<FBlueprintEditor> OwningEditor;
-	FDiffPanel PanelLocal;
-	TWeakPtr<SWindow>		OwningWindow;
+	FBlueprintMergeData Data;
+
+	TSharedPtr<SWidget>		GraphView;
+	TSharedPtr<SWidget>		ComponentsView;
+	TSharedPtr<SWidget>		DefaultsView;
 
 	// This has to be allocated here because SListView cannot own the list
 	// that it is displaying. It also seems like the display list *has*
 	// to be a list of TSharedPtrs.
-	TArray< TSharedPtr<FDiffSingleResult> > LocalDiffResults;
-	TArray< TSharedPtr<FDiffSingleResult> > RemoteDiffResults;
+	TArray< TSharedPtr<struct FDiffSingleResult> > LocalDiffResults;
+	TArray< TSharedPtr<struct FDiffSingleResult> > RemoteDiffResults;
 };
diff --git a/Engine/Source/Developer/Merge/Public/Merge.h b/Engine/Source/Developer/Merge/Public/Merge.h
--- a/Engine/Source/Developer/Merge/Public/Merge.h
+++ b/Engine/Source/Developer/Merge/Public/Merge.h
@@ -3,17 +3,10 @@
 #pragma once
 
 #include "ModuleManager.h"
-#include "IAssetTypeActions.h"
 
 class UBlueprint;
 class FBlueprintEditor;
 
-struct FMergeDisplayArgs
-{
-	FRevisionInfo DepotHeadRevision;
-	FRevisionInfo CommonBaseRevision;
-};
-
 /**
  * The public interface to this module
  */
@@ -48,7 +41,7 @@ public:
 	 *
 	 * @return The merge view widget
 	 */
-	virtual TSharedRef<class SDockTab> GenerateMergeWidget( const UBlueprint& Object, TSharedRef<FBlueprintEditor> Editor ) = 0;
+	virtual TSharedRef<class SDockTab> GenerateMergeWidget( const UBlueprint& Object, TSharedRef< class FBlueprintEditor> Editor ) = 0;
 
 	/** 
 	 * @return whether the blueprint is in a conflicted state, and can therefore be merged.
diff --git a/Engine/Source/Editor/Kismet/Private/SBlueprintDiff.cpp b/Engine/Source/Editor/Kismet/Private/SBlueprintDiff.cpp
--- a/Engine/Source/Editor/Kismet/Private/SBlueprintDiff.cpp
+++ b/Engine/Source/Editor/Kismet/Private/SBlueprintDiff.cpp
@@ -79,29 +79,21 @@ const UObject* GetCDO( const UBlueprint* ForBlueprint )
 	return ForBlueprint->GeneratedClass->ClassDefaultObject;
 }
 
-/** Individual Diff item shown in the list of diffs */
-struct FDiffResultItem: public TSharedFromThis<FDiffResultItem>
+TSharedRef<SWidget>	FDiffResultItem::GenerateWidget() const
 {
-	FDiffResultItem(FDiffSingleResult InResult): Result(InResult){}
-
-	FDiffSingleResult Result;
-
-	TSharedRef<SWidget>	GenerateWidget() const
+	FString ToolTip = Result.ToolTip;
+	FLinearColor Color = Result.DisplayColor;
+	FString Text = Result.DisplayString;
+	if (Text.Len() == 0)
 	{
-		FString ToolTip = Result.ToolTip;
-		FLinearColor Color = Result.DisplayColor;
-		FString Text = Result.DisplayString;
-		if(Text.Len() == 0)
-		{
-			Text = LOCTEXT("DIF_UnknownDiff", "Unknown Diff").ToString();
-			ToolTip = LOCTEXT("DIF_Confused", "There is an unspecified difference").ToString();
-		}
-		return SNew(STextBlock)
-				.ToolTipText(ToolTip)
-				.ColorAndOpacity(Color)
-				.Text(Text);
+		Text = LOCTEXT("DIF_UnknownDiff", "Unknown Diff").ToString();
+		ToolTip = LOCTEXT("DIF_Confused", "There is an unspecified difference").ToString();
 	}
-};
+	return SNew(STextBlock)
+		.ToolTipText(ToolTip)
+		.ColorAndOpacity(Color)
+		.Text(Text);
+}
 
 FListItemGraphToDiff::FListItemGraphToDiff( class SBlueprintDiff* InDiff, class UEdGraph* InGraphOld, class UEdGraph* InGraphNew, const FRevisionInfo& InRevisionOld, const FRevisionInfo& InRevisionNew )
 	: Diff(InDiff), GraphOld(InGraphOld), GraphNew(InGraphNew), RevisionOld(InRevisionOld), RevisionNew(InRevisionNew)
@@ -397,6 +389,7 @@ void FListItemGraphToDiff::KeyWasPressed( const FKeyboardEvent& InKeyboardEvent
 FDiffPanel::FDiffPanel()
 {
 	Blueprint = NULL;
+	LastFocusedPin = NULL;
 }
 
 // Generates a widget showing CDO values. NewBlueprint's CDO values will be displayed by the widget (Base just used for determining what to hide):
@@ -487,8 +480,6 @@ static TSharedRef<SWidget> GenerateComponentsDiff( UBlueprint const* BaseBluepri
 BEGIN_SLATE_FUNCTION_BUILD_OPTIMIZATION
 void SBlueprintDiff::Construct( const FArguments& InArgs)
 {
-	LastPinTarget = NULL;
-	LastOtherPinTarget = NULL;
 	FDiffListCommands::Register();
 	check(InArgs._BlueprintOld && InArgs._BlueprintNew);
 	PanelOld.Blueprint = InArgs._BlueprintOld;
@@ -631,9 +622,20 @@ void SBlueprintDiff::OnGraphChanged(FListItemGraphToDiff* Diff)
 	}
 }
 
+TSharedRef<SWidget> SBlueprintDiff::DefaultEmptyPanel()
+{
+	return SNew(SHorizontalBox)
+		+ SHorizontalBox::Slot()
+		.HAlign(HAlign_Center)
+		.VAlign(VAlign_Center)
+		[
+			SNew(STextBlock)
+			.Text(LOCTEXT("BlueprintDifGraphsToolTip", "Select Graph to Diff"))
+		];
+}
+
 void SBlueprintDiff::FocusOnGraphRevisions( class UEdGraph* GraphOld, class UEdGraph* GraphNew, FListItemGraphToDiff* Diff )
 {
-	DisablePinDiffFocus();
 	HandleGraphChanged( GraphOld ? GraphOld->GetName() : GraphNew->GetName() );
 
 	ResetGraphEditors();
@@ -644,41 +646,36 @@ void SBlueprintDiff::FocusOnGraphRevisions( class UEdGraph* GraphOld, class UEdG
 
 void SBlueprintDiff::OnDiffListSelectionChanged(const TSharedPtr<struct FDiffResultItem>& TheDiff, FListItemGraphToDiff* GraphDiffer )
 {
-	DisablePinDiffFocus();
-
-	//focus the graph onto the diff that was clicked on
 	FDiffSingleResult Result = TheDiff->Result;
-	if(Result.Pin1)
+
+	const auto SafeClearSelection = []( TWeakPtr<SGraphEditor> GraphEditor )
 	{
-		PanelNew.GraphEditor.Pin()->ClearSelectionSet();
-		PanelOld.GraphEditor.Pin()->ClearSelectionSet();
-	
-		if(Result.Pin1)
+		auto GraphEditorPtr = GraphEditor.Pin();
+		if( GraphEditorPtr.IsValid())
 		{
-			LastPinTarget = Result.Pin1;
-			Result.Pin1->bIsDiffing = true;
-			GetGraphEditorForGraph(Result.Pin1->GetOwningNode()->GetGraph())->JumpToPin(Result.Pin1);
+			GraphEditorPtr->ClearSelectionSet();
 		}
+	};
 
-		if(Result.Pin2)
+	SafeClearSelection( PanelNew.GraphEditor );
+	SafeClearSelection( PanelOld.GraphEditor );
+
+	if (Result.Pin1)
+	{
+		GetDiffPanelForNode(*Result.Pin1->GetOwningNode()).FocusDiff(*Result.Pin1);
+		if (Result.Pin2)
 		{
-			LastOtherPinTarget = Result.Pin2;
-			Result.Pin2->bIsDiffing = true;
-			GetGraphEditorForGraph(Result.Pin2->GetOwningNode()->GetGraph())->JumpToPin(Result.Pin2);
+			GetDiffPanelForNode(*Result.Pin2->GetOwningNode()).FocusDiff(*Result.Pin2);
 		}
 	}
-	else if(Result.Node1)
+	else if (Result.Node1)
 	{
-		PanelNew.GraphEditor.Pin()->ClearSelectionSet();
-		PanelOld.GraphEditor.Pin()->ClearSelectionSet();
-
-		if(Result.Node2)
+		GetDiffPanelForNode(*Result.Node1).FocusDiff(*Result.Node1);
+		if (Result.Node2)
 		{
-			GetGraphEditorForGraph(Result.Node2->GetGraph())->JumpToNode(Result.Node2, false);
+			GetDiffPanelForNode(*Result.Node2).FocusDiff(*Result.Node2);
 		}
-		GetGraphEditorForGraph(Result.Node1->GetGraph())->JumpToNode(Result.Node1, false);
 	}
-
 }
 
 FReply	SBlueprintDiff::OnToggleLockView()
@@ -712,6 +709,12 @@ void SBlueprintDiff::ResetGraphEditors()
 
 void FDiffPanel::GeneratePanel(UEdGraph* Graph, UEdGraph* GraphToDiff)
 {
+	if( LastFocusedPin )
+	{
+		LastFocusedPin->bIsDiffing = false;
+	}
+	LastFocusedPin = NULL;
+
 	TSharedPtr<SWidget> Widget = SNew(SBorder)
 								.HAlign(HAlign_Center)
 								.VAlign(VAlign_Center)
@@ -824,30 +827,44 @@ bool FDiffPanel::CanCopyNodes() const
 	return false;
 }
 
-SGraphEditor* SBlueprintDiff::GetGraphEditorForGraph( UEdGraph* Graph ) const
+void FDiffPanel::FocusDiff(UEdGraphPin& Pin)
 {
-	if(PanelOld.GraphEditor.Pin()->GetCurrentGraph() == Graph)
+	if( LastFocusedPin )
 	{
-		return PanelOld.GraphEditor.Pin().Get();
+		LastFocusedPin->bIsDiffing = false;
 	}
-	else if(PanelNew.GraphEditor.Pin()->GetCurrentGraph() == Graph)
+	Pin.bIsDiffing = true;
+	LastFocusedPin = &Pin;
+
+	GraphEditor.Pin()->JumpToPin(&Pin);
+}
+
+void FDiffPanel::FocusDiff(UEdGraphNode& Node)
+{
+	if (LastFocusedPin)
 	{
-		return PanelNew.GraphEditor.Pin().Get();
+		LastFocusedPin->bIsDiffing = false;
 	}
-	checkNoEntry();
-	return NULL;
+	LastFocusedPin = NULL;
+
+	GraphEditor.Pin()->JumpToNode(&Node, false);
 }
 
-TSharedRef<SWidget> SBlueprintDiff::DefaultEmptyPanel()
+FDiffPanel& SBlueprintDiff::GetDiffPanelForNode(UEdGraphNode& Node)
 {
-	return SNew(SHorizontalBox)
-		+ SHorizontalBox::Slot()
-		.HAlign(HAlign_Center)
-		.VAlign(VAlign_Center)
-		[
-			SNew(STextBlock)
-			.Text(LOCTEXT("BlueprintDifGraphsToolTip", "Select Graph to Diff"))
-		];
+	auto OldGraphEditorPtr = PanelOld.GraphEditor.Pin();
+	if (OldGraphEditorPtr.IsValid() && Node.GetGraph() == OldGraphEditorPtr->GetCurrentGraph())
+	{
+		return PanelOld;
+	}
+	auto NewGraphEditorPtr = PanelNew.GraphEditor.Pin();
+	if (NewGraphEditorPtr.IsValid() && Node.GetGraph() == NewGraphEditorPtr->GetCurrentGraph())
+	{
+		return PanelNew;
+	}
+	checkf(false, TEXT("Looking for node %s but it cannot be found in provided panels"), *Node.GetName());
+	static FDiffPanel Default;
+	return Default;
 }
 
 TSharedRef<SWidget> SBlueprintDiff::GenerateDiffWindow()
@@ -1082,17 +1099,5 @@ FReply SBlueprintDiff::OnKeyDown( const FGeometry& MyGeometry, const FKeyboardEv
 	return FReply::Handled();
 }
 
-void SBlueprintDiff::DisablePinDiffFocus()
-{
-	if(LastPinTarget)
-	{
-		LastPinTarget->bIsDiffing = false;
-	}
-	if(LastOtherPinTarget)
-	{
-		LastOtherPinTarget->bIsDiffing = false;
-	}
-}
-
 #undef LOCTEXT_NAMESPACE
 
diff --git a/Engine/Source/Editor/Kismet/Public/SBlueprintDiff.h b/Engine/Source/Editor/Kismet/Public/SBlueprintDiff.h
--- a/Engine/Source/Editor/Kismet/Public/SBlueprintDiff.h
+++ b/Engine/Source/Editor/Kismet/Public/SBlueprintDiff.h
@@ -2,6 +2,7 @@
 #pragma once
 
 #include "Developer/AssetTools/Public/IAssetTypeActions.h"
+#include "Editor/GraphEditor/Public/DiffResults.h"
 #include "GraphEditor.h"
 
 struct FMatchName
@@ -19,8 +20,18 @@ struct FMatchName
 	FString const& Name;
 };
 
+/** Individual Diff item shown in the list of diffs */
+struct FDiffResultItem : public TSharedFromThis<FDiffResultItem>
+{
+	FDiffResultItem(FDiffSingleResult InResult) : Result(InResult){}
+
+	FDiffSingleResult Result;
+
+	TSharedRef<SWidget> KISMET_API GenerateWidget() const;
+};
+
 /*List item that entry for a graph*/
-struct FListItemGraphToDiff: public TSharedFromThis<FListItemGraphToDiff>
+struct KISMET_API FListItemGraphToDiff: public TSharedFromThis<FListItemGraphToDiff>
 {
 	FListItemGraphToDiff(class SBlueprintDiff* Diff, class UEdGraph* GraphOld, class UEdGraph* GraphNew, const FRevisionInfo& RevisionOld, const FRevisionInfo& RevisionNew);
 	~FListItemGraphToDiff();
@@ -122,6 +133,10 @@ struct KISMET_API FDiffPanel
 	/*Can user copy any of the selected nodes?*/
 	bool CanCopyNodes() const;
 
+	/*Functions used to focus/find a particular change in a diff result*/
+	void FocusDiff(UEdGraphPin& Pin);
+	void FocusDiff(UEdGraphNode& Node);
+
 	/*The blueprint that owns the graph we are showing*/
 	const class UBlueprint*				Blueprint;
 
@@ -137,6 +152,8 @@ struct KISMET_API FDiffPanel
 	/*A name identifying which asset this panel is displaying */
 	bool							bShowAssetName;
 
+	/*The panel stores the last pin that was focused on by the user, so that it can clear the visual style when selection changes*/
+	UEdGraphPin*					LastFocusedPin;
 private:
 	/*Command list for this diff panel*/
 	TSharedPtr<FUICommandList> GraphEditorCommands;
@@ -162,6 +179,9 @@ public:
 	/** Called when a new Graph is clicked on by user */
 	void OnGraphChanged(FListItemGraphToDiff* Diff);
 
+	/** Helper function for generating an empty widget */
+	static TSharedRef<SWidget> DefaultEmptyPanel();
+
 protected:
 	/* Need to process keys for shortcuts to buttons */
 	virtual FReply OnKeyDown( const FGeometry& MyGeometry, const FKeyboardEvent& InKeyboardEvent ) override;
@@ -200,10 +220,7 @@ protected:
 	TArray< FGraphToDiff> Graphs;
 
 	/** Get Graph editor associated with this Graph */
-	virtual SGraphEditor* GetGraphEditorForGraph(UEdGraph* Graph) const ;
-
-	/** Helper function for generating an empty widget */
-	static TSharedRef<SWidget> DefaultEmptyPanel();
+	FDiffPanel& GetDiffPanelForNode(UEdGraphNode& Node);
 
 	/** Generates the widgets used to display the graphs that are being diff'd */
 	virtual TSharedRef<SWidget> GenerateDiffWindow();
@@ -239,12 +256,6 @@ protected:
 	/** The ListView containing the graphs the user can select */
 	TSharedPtr<SListViewType>	GraphsToDiff;
 
-	/** The last pin the user clicked on */
-	UEdGraphPin* LastPinTarget;
-
-	/** The last other pin the user clicked on */
-	UEdGraphPin* LastOtherPinTarget;
-
 	friend struct FListItemGraphToDiff;
 };
 
