#UE4 build fix, missing file

diff --git a/Engine/Source/Developer/Merge/Private/BlueprintMergeData.h b/Engine/Source/Developer/Merge/Private/BlueprintMergeData.h
new file mode 100644
--- /dev/null
+++ b/Engine/Source/Developer/Merge/Private/BlueprintMergeData.h
@@ -0,0 +1,45 @@
+// Copyright 1998-2014 Epic Games, Inc. All Rights Reserved.
+#pragma once
+
+#include "IAssetTypeActions.h"
+
+// This data is shared by the various controls (graph view, components view, and defaults view) 
+// that are presented as part of the merge tool. Each view keeps its own shallow copy of the data:
+struct FBlueprintMergeData
+{
+	FBlueprintMergeData()
+		: OwningEditor()
+		, BlueprintLocal(NULL)
+		, BlueprintBase(NULL)
+		, RevisionBase()
+		, BlueprintNew(NULL)
+		, RevisionNew()
+	{
+	}
+
+	FBlueprintMergeData( 
+		TWeakPtr<class FBlueprintEditor>	InOwningEditor
+		, const class UBlueprint*			InBlueprintLocal
+		, const class UBlueprint*			InBlueprintBase
+		, FRevisionInfo						InRevisionBase
+		, const class UBlueprint*			InBlueprintNew
+		, FRevisionInfo						InRevisionNew
+	)
+		: OwningEditor(		InOwningEditor	)
+		, BlueprintLocal(	InBlueprintLocal)
+		, BlueprintBase(	InBlueprintBase	)
+		, RevisionBase(		InRevisionBase	)
+		, BlueprintNew(		InBlueprintNew	)
+		, RevisionNew(		InRevisionNew	)
+	{
+	}
+
+	TWeakPtr<class FBlueprintEditor>	OwningEditor;
+	
+	const class UBlueprint*				BlueprintLocal;
+	const class UBlueprint*				BlueprintBase;
+	FRevisionInfo						RevisionBase;
+
+	const class UBlueprint*				BlueprintNew;
+	FRevisionInfo						RevisionNew;
+};
