#UE4 moved SMergeGraphView into its own file

diff --git a/Engine/Source/Developer/Merge/Private/SBlueprintMerge.cpp b/Engine/Source/Developer/Merge/Private/SBlueprintMerge.cpp
--- a/Engine/Source/Developer/Merge/Private/SBlueprintMerge.cpp
+++ b/Engine/Source/Developer/Merge/Private/SBlueprintMerge.cpp
@@ -2,492 +2,12 @@
 
 #include "MergePrivatePCH.h"
 #include "BlueprintEditor.h"
-#include "DiffResults.h"
-#include "GraphDiffControl.h"
-#include "ISourceControlModule.h"
-#include "SBlueprintDiff.h"
 #include "SBlueprintMerge.h"
+#include "SMergeGraphView.h"
+#include "ISourceControlModule.h"
 
 #define LOCTEXT_NAMESPACE "SBlueprintMerge"
 
-// This cludge is used to move arrays of structures into arrays of 
-// shared ptrs of those structures. It is used because SListView
-// does not support arrays of values:
-template< typename T >
-TArray< TSharedPtr<T> > HackToShared( const TArray<T>& Values )
-{
-	TArray< TSharedPtr<T> > Ret;
-	for( const auto& Value : Values )
-	{
-		Ret.Add( TSharedPtr<T>( new T(Value) ) );
-	}
-	return Ret;
-}
-
-namespace EMergeParticipant
-{
-	enum Type
-	{
-		MERGE_PARTICIPANT_REMOTE,
-		MERGE_PARTICIPANT_BASE,
-		MERGE_PARTICIPANT_LOCAL,
-	};
-}
-
-struct FBlueprintRevPair
-{
-	FBlueprintRevPair( const UBlueprint* InBlueprint, const FRevisionInfo& InRevData )
-		: Blueprint( InBlueprint )
-		, RevData( InRevData )
-	{
-	}
-
-	const UBlueprint* Blueprint;
-	const FRevisionInfo& RevData;
-};
-
-struct FMergeGraphRowEntry
-{
-	FString GraphName;
-
-	// These lists contain shared ptrs because they are displayed by a
-	// SListView, and it currently does not support lists of values.
-	TArray< TSharedPtr<FDiffSingleResult> > RemoteDifferences;
-	TArray< TSharedPtr<FDiffSingleResult> > LocalDifferences;
-	
-	bool bExistsInRemote;
-	bool bExistsInBase;
-	bool bExistsInLocal;
-};
-
-static UEdGraph* FindGraphByName(UBlueprint const& FromBlueprint, const FString& GraphName)
-{
-	TArray<UEdGraph*> Graphs;
-	FromBlueprint.GetAllGraphs(Graphs);
-
-	UEdGraph* Ret = NULL;
-	if (UEdGraph** Result = Graphs.FindByPredicate(FMatchName(GraphName)))
-	{
-		Ret = *Result;
-	}
-	return Ret;
-}
-
-static FDiffPanel InitializePanel(const FBlueprintRevPair& BlueprintRevPair)
-{
-	FDiffPanel Ret;
-	Ret.Blueprint = BlueprintRevPair.Blueprint;
-	SAssignNew(Ret.GraphEditorBorder, SBorder)
-		.VAlign(VAlign_Fill)
-		[
-			SBlueprintDiff::DefaultEmptyPanel()
-		];
-	Ret.RevisionInfo = BlueprintRevPair.RevData;
-	Ret.bShowAssetName = false;
-	return Ret;
-}
-
-// @todo doc: can i avoid shared ptrs for individual entries?
-static TArray< TSharedPtr<FMergeGraphRowEntry> > GenerateDiffListItems( const FBlueprintRevPair& RemoteBlueprint, const FBlueprintRevPair& BaseBlueprint, const FBlueprintRevPair& LocalBlueprint)
-{
-	// Index all the graphs by name, we use the name of the graph as the 
-	// basis of comparison between the various versions of the blueprint.
-	TMap< FString, UEdGraph* > RemoteGraphMap, BaseGraphMap, LocalGraphMap;
-	// We also want the set of all graph names in these blueprints, so that we 
-	// can iterate over every graph.
-	TSet< FString > AllGraphNames;
-	{ 
-		TArray<UEdGraph*> GraphsRemote, GraphsBase, GraphsLocal;
-		RemoteBlueprint.Blueprint->GetAllGraphs(GraphsRemote);
-		BaseBlueprint.Blueprint->GetAllGraphs(GraphsBase);
-		LocalBlueprint.Blueprint->GetAllGraphs(GraphsLocal);
-
-		const auto ToMap = [&AllGraphNames]( const TArray<UEdGraph*>& InList, TMap<FString, UEdGraph*>& OutMap)
-		{
-			for( auto Graph : InList )
-			{
-				OutMap.Add( Graph->GetName(), Graph );
-				AllGraphNames.Add(Graph->GetName());
-			}
-		};
-		ToMap(GraphsRemote, RemoteGraphMap);
-		ToMap(GraphsBase, BaseGraphMap);
-		ToMap(GraphsLocal, LocalGraphMap);
-	}
-	
-	TArray< TSharedPtr<FMergeGraphRowEntry> > Ret;
-	{
-		const auto GenerateDifferences = []( UEdGraph* GraphNew, UEdGraph** GraphOld )
-		{
-			TArray<FDiffSingleResult> Results;
-			FGraphDiffControl::DiffGraphs(GraphOld ? *GraphOld : NULL, GraphNew, Results );
-			return Results;
-		};
-
-		for( const auto& GraphName : AllGraphNames )
-		{
-			TArray< FDiffSingleResult > RemoteDifferences;
-			TArray< FDiffSingleResult > LocalDifferences;
-			bool bExistsInRemote, bExistsInBase, bExistsInLocal;
-
-			{
-				UEdGraph** RemoteGraph = RemoteGraphMap.Find(GraphName);
-				UEdGraph** BaseGraph = BaseGraphMap.Find(GraphName);
-				UEdGraph** LocalGraph = LocalGraphMap.Find(GraphName);
-
-				if( RemoteGraph )
-				{
-					RemoteDifferences = GenerateDifferences(*RemoteGraph, BaseGraph);
-				}
-
-				if( LocalGraph )
-				{
-					LocalDifferences = GenerateDifferences(*LocalGraph, BaseGraph);
-				}
-				
-				bExistsInRemote = RemoteGraph != NULL;
-				bExistsInBase = BaseGraph != NULL;
-				bExistsInLocal = LocalGraph != NULL;
-			}
-
-			FMergeGraphRowEntry NewEntry = {
-				GraphName,
-				HackToShared(RemoteDifferences),
-				HackToShared(LocalDifferences),
-				bExistsInRemote,
-				bExistsInBase,
-				bExistsInLocal
-			};
-			Ret.Add( TSharedPtr<FMergeGraphRowEntry>( new FMergeGraphRowEntry(NewEntry) ) );
-		}
-	}
-
-	return Ret;
-}
-
-static void LockViews( TArray<FDiffPanel>& Views, bool bAreLocked )
-{
-	for( auto& Panel : Views )
-	{
-		auto GraphEditor = Panel.GraphEditor.Pin();
-		if( GraphEditor.IsValid() )
-		{
-			// lock this panel to ever other panel:
-			for (auto& OtherPanel : Views)
-			{
-				auto OtherGraphEditor = OtherPanel.GraphEditor.Pin();
-				if( OtherGraphEditor.IsValid() &&
-					OtherGraphEditor != GraphEditor )
-				{
-					if( bAreLocked )
-					{
-						GraphEditor->LockToGraphEditor(OtherGraphEditor);
-					}
-					else
-					{
-						GraphEditor->UnlockFromGraphEditor(OtherGraphEditor);
-					}
-				}
-			}
-		}
-	}
-}
-
-FDiffPanel& GetDiffPanelForNode( const UEdGraphNode& Node, TArray< FDiffPanel >& Panels )
-{
-	for( auto& Panel : Panels )
-	{
-		auto GraphEditor = Panel.GraphEditor.Pin();
-		if( GraphEditor.IsValid() )
-		{
-			if( Node.GetGraph() == GraphEditor->GetCurrentGraph() )
-			{
-				return Panel;
-			}
-		}
-	}
-	checkf(false, TEXT("Looking for node %s but it cannot be found in provided panels"), *Node.GetName() );
-	static FDiffPanel Default;
-	return Default;
-}
-
-class SMergeGraphView : public SCompoundWidget
-{
-public:
-	SLATE_BEGIN_ARGS(SMergeGraphView)
-	{}
-	SLATE_END_ARGS()
-
-	void Construct(const FArguments InArgs, const FBlueprintMergeData& InData);
-
-private:
-	void OnGraphListSelectionChanged(TSharedPtr<FMergeGraphRowEntry> Item, ESelectInfo::Type SelectionType);
-	void OnDiffListSelectionChanged(TSharedPtr<FDiffSingleResult> Item, ESelectInfo::Type SelectionType);
-
-	FDiffPanel& GetRemotePanel() { return DiffPanels[EMergeParticipant::MERGE_PARTICIPANT_REMOTE]; }
-	FDiffPanel& GetBasePanel() { return DiffPanels[EMergeParticipant::MERGE_PARTICIPANT_BASE]; }
-	FDiffPanel& GetLocalPanel() { return DiffPanels[EMergeParticipant::MERGE_PARTICIPANT_LOCAL]; }
-
-	FReply OnToggleLockView();
-	const FSlateBrush*  GetLockViewImage() const;
-
-	TArray< FDiffPanel > DiffPanels;
-	FBlueprintMergeData Data;
-
-	TArray< TSharedPtr<FMergeGraphRowEntry> > DifferencesFromBase;
-	TSharedPtr<SBorder> DiffResultsWidget;
-
-	bool bViewsAreLocked;
-};
-
-// Some color constants used by the merge view:
-const FLinearColor MergeBlue = FLinearColor(0.3f, 0.3f, 1.f);
-const FLinearColor MergeRed = FLinearColor(1.0f, 0.2f, 0.3f);
-const FLinearColor MergeYellow = FLinearColor::Yellow;
-const FLinearColor MergeWhite = FLinearColor::White;
-
-void SMergeGraphView::Construct( const FArguments InArgs, const FBlueprintMergeData& InData )
-{
-	Data = InData;
-	bViewsAreLocked = true;
-
-	TArray<FBlueprintRevPair> BlueprintsForDisplay;
-	// EMergeParticipant::MERGE_PARTICIPANT_REMOTE
-	BlueprintsForDisplay.Add(FBlueprintRevPair(InData.BlueprintNew, InData.RevisionNew));
-	// EMergeParticipant::MERGE_PARTICIPANT_BASE
-	BlueprintsForDisplay.Add(FBlueprintRevPair(InData.BlueprintBase, InData.RevisionBase));
-	// EMergeParticipant::MERGE_PARTICIPANT_LOCAL
-	BlueprintsForDisplay.Add( FBlueprintRevPair(InData.BlueprintLocal, FRevisionInfo()) );
-
-	for( const auto& BlueprintRevPair : BlueprintsForDisplay )
-	{
-		DiffPanels.Add( InitializePanel( BlueprintRevPair ) );
-	}
-
-	TSharedRef<SSplitter> PanelContainer = SNew(SSplitter);
-	for( const auto& Panel : DiffPanels )
-	{
-		PanelContainer->AddSlot()[ Panel.GraphEditorBorder.ToSharedRef() ];
-	}
-
-	DifferencesFromBase = GenerateDiffListItems(	BlueprintsForDisplay[ EMergeParticipant::MERGE_PARTICIPANT_REMOTE ]
-													,BlueprintsForDisplay[ EMergeParticipant::MERGE_PARTICIPANT_BASE ]
-													,BlueprintsForDisplay[ EMergeParticipant::MERGE_PARTICIPANT_LOCAL ]);
-
-	// This is the function we'll use to generate a row in the control that lists all the available graphs:
-	const auto RowGenerator = [](TSharedPtr<FMergeGraphRowEntry> ParamItem, const TSharedRef<STableViewBase>& OwnerTable) -> TSharedRef<ITableRow>
-	{
-		// blue indicates added, red indicates changed, yellow indicates removed, white indicates no change:
-		const auto ComputeColorForRevision = [](bool bExistsInBase, bool bExistsInThisRevision, bool bHasChangesInThisRevision) -> FLinearColor
-		{
-			if( bExistsInBase )
-			{
-				if( !bExistsInThisRevision )
-				{
-					return MergeYellow;
-				}
-				else if( bHasChangesInThisRevision )
-				{
-					return MergeRed;
-				}
-			}
-			if( !bExistsInBase )
-			{
-				if( bExistsInThisRevision  )
-				{
-					return MergeBlue;
-				}
-			}
-
-			return MergeWhite;
-		};
-
-		FLinearColor RemoteColor = ComputeColorForRevision(ParamItem->bExistsInBase, ParamItem->bExistsInRemote, ParamItem->RemoteDifferences.Num() != 0);
-		FLinearColor BaseColor = ComputeColorForRevision(ParamItem->bExistsInBase, ParamItem->bExistsInBase, false);
-		FLinearColor LocalColor = ComputeColorForRevision(ParamItem->bExistsInBase, ParamItem->bExistsInLocal, ParamItem->LocalDifferences.Num() != 0);
-
-		const auto Box = [] (bool bIsPresent, FLinearColor Color) -> SHorizontalBox::FSlot&
-		{
-			return SHorizontalBox::Slot()
-				.AutoWidth()
-				.HAlign(HAlign_Right)
-				.VAlign(VAlign_Center)
-				.MaxWidth(8.0f)
-				[
-					SNew(SImage)
-					.ColorAndOpacity(Color)
-					.Image(bIsPresent ? FEditorStyle::GetBrush("BlueprintDif.HasGraph") : FEditorStyle::GetBrush("BlueprintDif.MissingGraph"))
-				];
-		};
-
-		return SNew(STableRow< TSharedPtr<FMergeGraphRowEntry> >, OwnerTable)
-			.ToolTipText(LOCTEXT("MergeGraphsDifferentText", "@todo doc"))
-			.Content()
-			[
-				SNew(SHorizontalBox)
-					+ SHorizontalBox::Slot()
-					[
-						SNew(STextBlock)
-						.ColorAndOpacity(FLinearColor::White)
-						.Text(ParamItem->GraphName)
-					]
-					+ Box(ParamItem->bExistsInRemote, RemoteColor )
-					+ Box(ParamItem->bExistsInBase, BaseColor)
-					+ Box(ParamItem->bExistsInLocal, LocalColor)
-			];
-	};
-
-	ChildSlot
-	[
-		SNew(SSplitter)
-		.Orientation(Orient_Horizontal)
-		+ SSplitter::Slot()
-		.Value(0.1f)
-		[
-			// ToolPanel.GroupBorder is currently what gives this control its dark grey appearance:
-			SNew(SBorder)
-			.BorderImage(FEditorStyle::GetBrush("NoBorder"))
-			.Padding(FMargin(1.0f))
-			[
-				SNew(SSplitter)
-				.Orientation(Orient_Vertical)
-				+SSplitter::Slot()
-				.Value(.3f)
-				[
-					SNew(SBorder)
-					.BorderImage(FEditorStyle::GetBrush("ToolPanel.GroupBorder"))
-					[
-						SNew(SVerticalBox)
-						+ SVerticalBox::Slot()
-						.AutoHeight()
-						.HAlign(HAlign_Left)
-						[
-							SNew(SButton)
-							.OnClicked(this, &SMergeGraphView::OnToggleLockView)
-							.Content()
-							[
-								SNew(SImage)
-								.Image(this, &SMergeGraphView::GetLockViewImage)
-								.ToolTipText(LOCTEXT("BPDifLock", "Lock the blueprint views?"))
-							]
-						]
-						+ SVerticalBox::Slot()
-						[
-							SNew(SListView< TSharedPtr<FMergeGraphRowEntry> >)
-							.ItemHeight(24)
-							.ListItemsSource(&DifferencesFromBase)
-							.OnGenerateRow(SListView< TSharedPtr<FMergeGraphRowEntry> >::FOnGenerateRow::CreateStatic(RowGenerator))
-							.SelectionMode(ESelectionMode::Single)
-							.OnSelectionChanged(this, &SMergeGraphView::OnGraphListSelectionChanged)
-						]
-					]
-				]
-				+ SSplitter::Slot()
-				.Value(.7f)
-				[
-					SAssignNew(DiffResultsWidget, SBorder)
-					.BorderImage(FEditorStyle::GetBrush("ToolPanel.GroupBorder"))
-				]
-			]
-		]
-		+ SSplitter::Slot()
-		.Value(0.9f)
-		[
-			PanelContainer
-		]
-	];
-}
-
-void SMergeGraphView::OnGraphListSelectionChanged(TSharedPtr<FMergeGraphRowEntry> Item, ESelectInfo::Type SelectionType)
-{
-	if (SelectionType != ESelectInfo::OnMouseClick || !Item.IsValid())
-	{
-		return;
-	}
-
-	FString GraphName = Item->GraphName;
-
-	UEdGraph* GraphRemote = FindGraphByName(*GetRemotePanel().Blueprint, GraphName);
-	UEdGraph* GraphBase = FindGraphByName(*GetBasePanel().Blueprint, GraphName);
-	UEdGraph* GraphLocal = FindGraphByName(*GetLocalPanel().Blueprint, GraphName);
-
-	GetBasePanel().GeneratePanel(GraphBase, NULL);
-	GetRemotePanel().GeneratePanel(GraphRemote, GraphBase);
-	GetLocalPanel().GeneratePanel(GraphLocal, GraphBase);
-
-	LockViews(DiffPanels, bViewsAreLocked);
-
-	DiffResultsWidget->SetContent(
-		SNew( SListView< TSharedPtr<FDiffSingleResult> >)
-		.ItemHeight(24)
-		.ListItemsSource(&Item->LocalDifferences)
-		.OnGenerateRow( SListView< TSharedPtr<FDiffSingleResult> >::FOnGenerateRow::CreateStatic(
-			[](TSharedPtr<FDiffSingleResult> ParamItem, const TSharedRef<STableViewBase>& OwnerTable) -> TSharedRef<ITableRow>
-			{
-				FDiffResultItem WidgetGenerator( *ParamItem );
-				return SNew(STableRow< TSharedPtr<FDiffSingleResult> >, OwnerTable)
-					.Content()
-					[
-						WidgetGenerator.GenerateWidget()
-					];
-			}
-		) )
-		.SelectionMode(ESelectionMode::Single)
-		.OnSelectionChanged(this, &SMergeGraphView::OnDiffListSelectionChanged)
-	);
-}
-
-void SMergeGraphView::OnDiffListSelectionChanged(TSharedPtr<FDiffSingleResult> Item, ESelectInfo::Type SelectionType)
-{
-	if( SelectionType != ESelectInfo::OnMouseClick || !Item.IsValid() )
-	{
-		return;
-	}
-
-	// Clear graph selection:
-	for( auto& Panel : DiffPanels )
-	{
-		auto GraphEdtiorPtr = Panel.GraphEditor.Pin();
-		if (GraphEdtiorPtr.IsValid())
-		{
-			GraphEdtiorPtr->ClearSelectionSet();
-		}
-	}
-
-	//focus the graph onto the diff that was clicked on
-	if (Item->Pin1)
-	{
-		GetDiffPanelForNode(*Item->Pin1->GetOwningNode(), DiffPanels).FocusDiff( *Item->Pin1 );
-
-		if (Item->Pin2)
-		{
-			GetDiffPanelForNode(*Item->Pin2->GetOwningNode(), DiffPanels).FocusDiff( *Item->Pin2 );
-		}
-	}
-	else if (Item->Node1)
-	{
-		GetDiffPanelForNode(*Item->Node1, DiffPanels).FocusDiff( *Item->Node1 );
-
-		if (Item->Node2)
-		{
-			GetDiffPanelForNode(*Item->Node2, DiffPanels).FocusDiff( *Item->Node2 );
-		}
-	}
-}
-
-FReply SMergeGraphView::OnToggleLockView()
-{
-	bViewsAreLocked = !bViewsAreLocked;
-
-	LockViews( DiffPanels, bViewsAreLocked );
-	return FReply::Handled();
-}
-
-const FSlateBrush* SMergeGraphView::GetLockViewImage() const
-{
-	return bViewsAreLocked ? FEditorStyle::GetBrush("GenericLock") : FEditorStyle::GetBrush("GenericUnlock");
-}
-
 /** 
  * Commands used to navigate the lists of differences between the local revision and the common base
  * or the remote revision and the common base.
diff --git a/Engine/Source/Developer/Merge/Private/SMergeGraphView.cpp b/Engine/Source/Developer/Merge/Private/SMergeGraphView.cpp
new file mode 100644
--- /dev/null
+++ b/Engine/Source/Developer/Merge/Private/SMergeGraphView.cpp
@@ -0,0 +1,447 @@
+// Copyright 1998-2014 Epic Games, Inc. All Rights Reserved.
+
+#include "MergePrivatePCH.h"
+#include "DiffResults.h"
+#include "GraphDiffControl.h"
+#include "SMergeGraphView.h"
+
+#define LOCTEXT_NAMESPACE "SMergeGraphView"
+
+// This cludge is used to move arrays of structures into arrays of 
+// shared ptrs of those structures. It is used because SListView
+// does not support arrays of values:
+template< typename T >
+TArray< TSharedPtr<T> > HackToShared(const TArray<T>& Values)
+{
+	TArray< TSharedPtr<T> > Ret;
+	for (const auto& Value : Values)
+	{
+		Ret.Add(TSharedPtr<T>(new T(Value)));
+	}
+	return Ret;
+}
+
+struct FBlueprintRevPair
+{
+	FBlueprintRevPair(const UBlueprint* InBlueprint, const FRevisionInfo& InRevData)
+	: Blueprint(InBlueprint)
+	, RevData(InRevData)
+	{
+	}
+
+	const UBlueprint* Blueprint;
+	const FRevisionInfo& RevData;
+};
+
+struct FMergeGraphRowEntry
+{
+	FString GraphName;
+
+	// These lists contain shared ptrs because they are displayed by a
+	// SListView, and it currently does not support lists of values.
+	TArray< TSharedPtr<FDiffSingleResult> > RemoteDifferences;
+	TArray< TSharedPtr<FDiffSingleResult> > LocalDifferences;
+
+	bool bExistsInRemote;
+	bool bExistsInBase;
+	bool bExistsInLocal;
+};
+
+static UEdGraph* FindGraphByName(UBlueprint const& FromBlueprint, const FString& GraphName)
+{
+	TArray<UEdGraph*> Graphs;
+	FromBlueprint.GetAllGraphs(Graphs);
+
+	UEdGraph* Ret = NULL;
+	if (UEdGraph** Result = Graphs.FindByPredicate(FMatchName(GraphName)))
+	{
+		Ret = *Result;
+	}
+	return Ret;
+}
+
+static FDiffPanel InitializePanel(const FBlueprintRevPair& BlueprintRevPair)
+{
+	FDiffPanel Ret;
+	Ret.Blueprint = BlueprintRevPair.Blueprint;
+	SAssignNew(Ret.GraphEditorBorder, SBorder)
+		.VAlign(VAlign_Fill)
+		[
+			SBlueprintDiff::DefaultEmptyPanel()
+		];
+	Ret.RevisionInfo = BlueprintRevPair.RevData;
+	Ret.bShowAssetName = false;
+	return Ret;
+}
+
+// @todo doc: can i avoid shared ptrs for individual entries?
+static TArray< TSharedPtr<FMergeGraphRowEntry> > GenerateDiffListItems(const FBlueprintRevPair& RemoteBlueprint, const FBlueprintRevPair& BaseBlueprint, const FBlueprintRevPair& LocalBlueprint)
+{
+	// Index all the graphs by name, we use the name of the graph as the 
+	// basis of comparison between the various versions of the blueprint.
+	TMap< FString, UEdGraph* > RemoteGraphMap, BaseGraphMap, LocalGraphMap;
+	// We also want the set of all graph names in these blueprints, so that we 
+	// can iterate over every graph.
+	TSet< FString > AllGraphNames;
+	{
+		TArray<UEdGraph*> GraphsRemote, GraphsBase, GraphsLocal;
+		RemoteBlueprint.Blueprint->GetAllGraphs(GraphsRemote);
+		BaseBlueprint.Blueprint->GetAllGraphs(GraphsBase);
+		LocalBlueprint.Blueprint->GetAllGraphs(GraphsLocal);
+
+		const auto ToMap = [&AllGraphNames](const TArray<UEdGraph*>& InList, TMap<FString, UEdGraph*>& OutMap)
+		{
+			for (auto Graph : InList)
+			{
+				OutMap.Add(Graph->GetName(), Graph);
+				AllGraphNames.Add(Graph->GetName());
+			}
+		};
+		ToMap(GraphsRemote, RemoteGraphMap);
+		ToMap(GraphsBase, BaseGraphMap);
+		ToMap(GraphsLocal, LocalGraphMap);
+	}
+
+	TArray< TSharedPtr<FMergeGraphRowEntry> > Ret;
+	{
+		const auto GenerateDifferences = [](UEdGraph* GraphNew, UEdGraph** GraphOld)
+		{
+			TArray<FDiffSingleResult> Results;
+			FGraphDiffControl::DiffGraphs(GraphOld ? *GraphOld : NULL, GraphNew, Results);
+			return Results;
+		};
+
+		for (const auto& GraphName : AllGraphNames)
+		{
+			TArray< FDiffSingleResult > RemoteDifferences;
+			TArray< FDiffSingleResult > LocalDifferences;
+			bool bExistsInRemote, bExistsInBase, bExistsInLocal;
+
+			{
+				UEdGraph** RemoteGraph = RemoteGraphMap.Find(GraphName);
+				UEdGraph** BaseGraph = BaseGraphMap.Find(GraphName);
+				UEdGraph** LocalGraph = LocalGraphMap.Find(GraphName);
+
+				if (RemoteGraph)
+				{
+					RemoteDifferences = GenerateDifferences(*RemoteGraph, BaseGraph);
+				}
+
+				if (LocalGraph)
+				{
+					LocalDifferences = GenerateDifferences(*LocalGraph, BaseGraph);
+				}
+
+				bExistsInRemote = RemoteGraph != NULL;
+				bExistsInBase = BaseGraph != NULL;
+				bExistsInLocal = LocalGraph != NULL;
+			}
+
+			FMergeGraphRowEntry NewEntry = {
+				GraphName,
+				HackToShared(RemoteDifferences),
+				HackToShared(LocalDifferences),
+				bExistsInRemote,
+				bExistsInBase,
+				bExistsInLocal
+			};
+			Ret.Add(TSharedPtr<FMergeGraphRowEntry>(new FMergeGraphRowEntry(NewEntry)));
+		}
+	}
+
+	return Ret;
+}
+
+static void LockViews(TArray<FDiffPanel>& Views, bool bAreLocked)
+{
+	for (auto& Panel : Views)
+	{
+		auto GraphEditor = Panel.GraphEditor.Pin();
+		if (GraphEditor.IsValid())
+		{
+			// lock this panel to ever other panel:
+			for (auto& OtherPanel : Views)
+			{
+				auto OtherGraphEditor = OtherPanel.GraphEditor.Pin();
+				if (OtherGraphEditor.IsValid() &&
+					OtherGraphEditor != GraphEditor)
+				{
+					if (bAreLocked)
+					{
+						GraphEditor->LockToGraphEditor(OtherGraphEditor);
+					}
+					else
+					{
+						GraphEditor->UnlockFromGraphEditor(OtherGraphEditor);
+					}
+				}
+			}
+		}
+	}
+}
+
+FDiffPanel& GetDiffPanelForNode(const UEdGraphNode& Node, TArray< FDiffPanel >& Panels)
+{
+	for (auto& Panel : Panels)
+	{
+		auto GraphEditor = Panel.GraphEditor.Pin();
+		if (GraphEditor.IsValid())
+		{
+			if (Node.GetGraph() == GraphEditor->GetCurrentGraph())
+			{
+				return Panel;
+			}
+		}
+	}
+	checkf(false, TEXT("Looking for node %s but it cannot be found in provided panels"), *Node.GetName());
+	static FDiffPanel Default;
+	return Default;
+}
+
+// Some color constants used by the merge view:
+const FLinearColor MergeBlue = FLinearColor(0.3f, 0.3f, 1.f);
+const FLinearColor MergeRed = FLinearColor(1.0f, 0.2f, 0.3f);
+const FLinearColor MergeYellow = FLinearColor::Yellow;
+const FLinearColor MergeWhite = FLinearColor::White;
+
+void SMergeGraphView::Construct(const FArguments InArgs, const FBlueprintMergeData& InData)
+{
+	Data = InData;
+	bViewsAreLocked = true;
+
+	TArray<FBlueprintRevPair> BlueprintsForDisplay;
+	// EMergeParticipant::MERGE_PARTICIPANT_REMOTE
+	BlueprintsForDisplay.Add(FBlueprintRevPair(InData.BlueprintNew, InData.RevisionNew));
+	// EMergeParticipant::MERGE_PARTICIPANT_BASE
+	BlueprintsForDisplay.Add(FBlueprintRevPair(InData.BlueprintBase, InData.RevisionBase));
+	// EMergeParticipant::MERGE_PARTICIPANT_LOCAL
+	BlueprintsForDisplay.Add(FBlueprintRevPair(InData.BlueprintLocal, FRevisionInfo()));
+
+	for (const auto& BlueprintRevPair : BlueprintsForDisplay)
+	{
+		DiffPanels.Add(InitializePanel(BlueprintRevPair));
+	}
+
+	TSharedRef<SSplitter> PanelContainer = SNew(SSplitter);
+	for (const auto& Panel : DiffPanels)
+	{
+		PanelContainer->AddSlot()[Panel.GraphEditorBorder.ToSharedRef()];
+	}
+
+	DifferencesFromBase = GenerateDiffListItems(BlueprintsForDisplay[EMergeParticipant::MERGE_PARTICIPANT_REMOTE]
+		, BlueprintsForDisplay[EMergeParticipant::MERGE_PARTICIPANT_BASE]
+		, BlueprintsForDisplay[EMergeParticipant::MERGE_PARTICIPANT_LOCAL]);
+
+	// This is the function we'll use to generate a row in the control that lists all the available graphs:
+	const auto RowGenerator = [](TSharedPtr<FMergeGraphRowEntry> ParamItem, const TSharedRef<STableViewBase>& OwnerTable) -> TSharedRef<ITableRow>
+	{
+		// blue indicates added, red indicates changed, yellow indicates removed, white indicates no change:
+		const auto ComputeColorForRevision = [](bool bExistsInBase, bool bExistsInThisRevision, bool bHasChangesInThisRevision) -> FLinearColor
+		{
+			if (bExistsInBase)
+			{
+				if (!bExistsInThisRevision)
+				{
+					return MergeYellow;
+				}
+				else if (bHasChangesInThisRevision)
+				{
+					return MergeRed;
+				}
+			}
+			if (!bExistsInBase)
+			{
+				if (bExistsInThisRevision)
+				{
+					return MergeBlue;
+				}
+			}
+
+			return MergeWhite;
+		};
+
+		FLinearColor RemoteColor = ComputeColorForRevision(ParamItem->bExistsInBase, ParamItem->bExistsInRemote, ParamItem->RemoteDifferences.Num() != 0);
+		FLinearColor BaseColor = ComputeColorForRevision(ParamItem->bExistsInBase, ParamItem->bExistsInBase, false);
+		FLinearColor LocalColor = ComputeColorForRevision(ParamItem->bExistsInBase, ParamItem->bExistsInLocal, ParamItem->LocalDifferences.Num() != 0);
+
+		const auto Box = [](bool bIsPresent, FLinearColor Color) -> SHorizontalBox::FSlot&
+		{
+			return SHorizontalBox::Slot()
+				.AutoWidth()
+				.HAlign(HAlign_Right)
+				.VAlign(VAlign_Center)
+				.MaxWidth(8.0f)
+				[
+					SNew(SImage)
+					.ColorAndOpacity(Color)
+					.Image(bIsPresent ? FEditorStyle::GetBrush("BlueprintDif.HasGraph") : FEditorStyle::GetBrush("BlueprintDif.MissingGraph"))
+				];
+		};
+
+		return SNew(STableRow< TSharedPtr<FMergeGraphRowEntry> >, OwnerTable)
+			.ToolTipText(LOCTEXT("MergeGraphsDifferentText", "@todo doc"))
+			.Content()
+			[
+				SNew(SHorizontalBox)
+				+ SHorizontalBox::Slot()
+				[
+					SNew(STextBlock)
+					.ColorAndOpacity(FLinearColor::White)
+					.Text(ParamItem->GraphName)
+				]
+				+ Box(ParamItem->bExistsInRemote, RemoteColor)
+					+ Box(ParamItem->bExistsInBase, BaseColor)
+					+ Box(ParamItem->bExistsInLocal, LocalColor)
+			];
+	};
+
+	ChildSlot
+		[
+			SNew(SSplitter)
+			.Orientation(Orient_Horizontal)
+			+ SSplitter::Slot()
+			.Value(0.1f)
+			[
+				// ToolPanel.GroupBorder is currently what gives this control its dark grey appearance:
+				SNew(SBorder)
+				.BorderImage(FEditorStyle::GetBrush("NoBorder"))
+				.Padding(FMargin(1.0f))
+				[
+					SNew(SSplitter)
+					.Orientation(Orient_Vertical)
+					+ SSplitter::Slot()
+					.Value(.3f)
+					[
+						SNew(SBorder)
+						.BorderImage(FEditorStyle::GetBrush("ToolPanel.GroupBorder"))
+						[
+							SNew(SVerticalBox)
+							+ SVerticalBox::Slot()
+							.AutoHeight()
+							.HAlign(HAlign_Left)
+							[
+								SNew(SButton)
+								.OnClicked(this, &SMergeGraphView::OnToggleLockView)
+								.Content()
+								[
+									SNew(SImage)
+									.Image(this, &SMergeGraphView::GetLockViewImage)
+									.ToolTipText(LOCTEXT("BPDifLock", "Lock the blueprint views?"))
+								]
+							]
+							+ SVerticalBox::Slot()
+								[
+									SNew(SListView< TSharedPtr<FMergeGraphRowEntry> >)
+									.ItemHeight(24)
+									.ListItemsSource(&DifferencesFromBase)
+									.OnGenerateRow(SListView< TSharedPtr<FMergeGraphRowEntry> >::FOnGenerateRow::CreateStatic(RowGenerator))
+									.SelectionMode(ESelectionMode::Single)
+									.OnSelectionChanged(this, &SMergeGraphView::OnGraphListSelectionChanged)
+								]
+						]
+					]
+					+ SSplitter::Slot()
+						.Value(.7f)
+						[
+							SAssignNew(DiffResultsWidget, SBorder)
+							.BorderImage(FEditorStyle::GetBrush("ToolPanel.GroupBorder"))
+						]
+				]
+			]
+			+ SSplitter::Slot()
+				.Value(0.9f)
+				[
+					PanelContainer
+				]
+		];
+}
+
+void SMergeGraphView::OnGraphListSelectionChanged(TSharedPtr<FMergeGraphRowEntry> Item, ESelectInfo::Type SelectionType)
+{
+	if (SelectionType != ESelectInfo::OnMouseClick || !Item.IsValid())
+	{
+		return;
+	}
+
+	FString GraphName = Item->GraphName;
+
+	UEdGraph* GraphRemote = FindGraphByName(*GetRemotePanel().Blueprint, GraphName);
+	UEdGraph* GraphBase = FindGraphByName(*GetBasePanel().Blueprint, GraphName);
+	UEdGraph* GraphLocal = FindGraphByName(*GetLocalPanel().Blueprint, GraphName);
+
+	GetBasePanel().GeneratePanel(GraphBase, NULL);
+	GetRemotePanel().GeneratePanel(GraphRemote, GraphBase);
+	GetLocalPanel().GeneratePanel(GraphLocal, GraphBase);
+
+	LockViews(DiffPanels, bViewsAreLocked);
+
+	DiffResultsWidget->SetContent(
+		SNew(SListView< TSharedPtr<FDiffSingleResult> >)
+		.ItemHeight(24)
+		.ListItemsSource(&Item->LocalDifferences)
+		.OnGenerateRow(SListView< TSharedPtr<FDiffSingleResult> >::FOnGenerateRow::CreateStatic(
+		[](TSharedPtr<FDiffSingleResult> ParamItem, const TSharedRef<STableViewBase>& OwnerTable) -> TSharedRef<ITableRow>
+	{
+		FDiffResultItem WidgetGenerator(*ParamItem);
+		return SNew(STableRow< TSharedPtr<FDiffSingleResult> >, OwnerTable)
+			.Content()
+			[
+				WidgetGenerator.GenerateWidget()
+			];
+	}
+	))
+		.SelectionMode(ESelectionMode::Single)
+		.OnSelectionChanged(this, &SMergeGraphView::OnDiffListSelectionChanged)
+		);
+}
+
+void SMergeGraphView::OnDiffListSelectionChanged(TSharedPtr<FDiffSingleResult> Item, ESelectInfo::Type SelectionType)
+{
+	if (SelectionType != ESelectInfo::OnMouseClick || !Item.IsValid())
+	{
+		return;
+	}
+
+	// Clear graph selection:
+	for (auto& Panel : DiffPanels)
+	{
+		auto GraphEdtiorPtr = Panel.GraphEditor.Pin();
+		if (GraphEdtiorPtr.IsValid())
+		{
+			GraphEdtiorPtr->ClearSelectionSet();
+		}
+	}
+
+	//focus the graph onto the diff that was clicked on
+	if (Item->Pin1)
+	{
+		GetDiffPanelForNode(*Item->Pin1->GetOwningNode(), DiffPanels).FocusDiff(*Item->Pin1);
+
+		if (Item->Pin2)
+		{
+			GetDiffPanelForNode(*Item->Pin2->GetOwningNode(), DiffPanels).FocusDiff(*Item->Pin2);
+		}
+	}
+	else if (Item->Node1)
+	{
+		GetDiffPanelForNode(*Item->Node1, DiffPanels).FocusDiff(*Item->Node1);
+
+		if (Item->Node2)
+		{
+			GetDiffPanelForNode(*Item->Node2, DiffPanels).FocusDiff(*Item->Node2);
+		}
+	}
+}
+
+FReply SMergeGraphView::OnToggleLockView()
+{
+	bViewsAreLocked = !bViewsAreLocked;
+
+	LockViews(DiffPanels, bViewsAreLocked);
+	return FReply::Handled();
+}
+
+const FSlateBrush* SMergeGraphView::GetLockViewImage() const
+{
+	return bViewsAreLocked ? FEditorStyle::GetBrush("GenericLock") : FEditorStyle::GetBrush("GenericUnlock");
+}
diff --git a/Engine/Source/Developer/Merge/Private/SMergeGraphView.h b/Engine/Source/Developer/Merge/Private/SMergeGraphView.h
new file mode 100644
--- /dev/null
+++ b/Engine/Source/Developer/Merge/Private/SMergeGraphView.h
@@ -0,0 +1,44 @@
+// Copyright 1998-2014 Epic Games, Inc. All Rights Reserved.
+#pragma once
+
+#include "BlueprintMergeData.h"
+#include "SBlueprintDiff.h"
+
+namespace EMergeParticipant
+{
+	enum Type
+	{
+		MERGE_PARTICIPANT_REMOTE,
+		MERGE_PARTICIPANT_BASE,
+		MERGE_PARTICIPANT_LOCAL,
+	};
+}
+
+class SMergeGraphView : public SCompoundWidget
+{
+public:
+	SLATE_BEGIN_ARGS(SMergeGraphView)
+	{}
+	SLATE_END_ARGS()
+
+	void Construct(const FArguments InArgs, const FBlueprintMergeData& InData);
+
+private:
+	void OnGraphListSelectionChanged(TSharedPtr<struct FMergeGraphRowEntry> Item, ESelectInfo::Type SelectionType);
+	void OnDiffListSelectionChanged(TSharedPtr<struct FDiffSingleResult> Item, ESelectInfo::Type SelectionType);
+
+	FDiffPanel& GetRemotePanel() { return DiffPanels[EMergeParticipant::MERGE_PARTICIPANT_REMOTE]; }
+	FDiffPanel& GetBasePanel() { return DiffPanels[EMergeParticipant::MERGE_PARTICIPANT_BASE]; }
+	FDiffPanel& GetLocalPanel() { return DiffPanels[EMergeParticipant::MERGE_PARTICIPANT_LOCAL]; }
+
+	FReply OnToggleLockView();
+	const FSlateBrush*  GetLockViewImage() const;
+
+	TArray< FDiffPanel > DiffPanels;
+	FBlueprintMergeData Data;
+
+	TArray< TSharedPtr< struct FMergeGraphRowEntry> > DifferencesFromBase;
+	TSharedPtr<class SBorder> DiffResultsWidget;
+
+	bool bViewsAreLocked;
+};
