#UE4 Merge tool now switches modes with the blueprint editor

diff --git a/Engine/Source/Developer/Merge/Private/SBlueprintMerge.cpp b/Engine/Source/Developer/Merge/Private/SBlueprintMerge.cpp
--- a/Engine/Source/Developer/Merge/Private/SBlueprintMerge.cpp
+++ b/Engine/Source/Developer/Merge/Private/SBlueprintMerge.cpp
@@ -1,40 +1,16 @@
 // Copyright 1998-2014 Epic Games, Inc. All Rights Reserved.
 
 #include "MergePrivatePCH.h"
+
 #include "BlueprintEditor.h"
+#include "BlueprintEditorModes.h"
+#include "ISourceControlModule.h"
+#include "SBlueprintEditorToolbar.h"
 #include "SBlueprintMerge.h"
 #include "SMergeGraphView.h"
-#include "ISourceControlModule.h"
 
 #define LOCTEXT_NAMESPACE "SBlueprintMerge"
 
-/** 
- * Commands used to navigate the lists of differences between the local revision and the common base
- * or the remote revision and the common base.
- */
-class FMergeDifferencesListCommands : public TCommands<FMergeDifferencesListCommands>
-{
-public:
-	FMergeDifferencesListCommands()
-		: TCommands<FMergeDifferencesListCommands>("MergeList", NSLOCTEXT("Merge", "Blueprint", "Blueprint Merge"), NAME_None, FEditorStyle::GetStyleSetName())
-	{
-	}
-
-	/** Go to previous or next difference, focusing on either the remote or local diff */
-	TSharedPtr<FUICommandInfo> LocalPrevious;
-	TSharedPtr<FUICommandInfo> LocalNext;
-	TSharedPtr<FUICommandInfo> RemotePrevious;
-	TSharedPtr<FUICommandInfo> RemoteNext;
-
-	virtual void RegisterCommands() override
-	{
-		UI_COMMAND(LocalPrevious, "LocalPrevious", "Go to previous local difference", EUserInterfaceActionType::Button, FInputGesture(EKeys::F7, EModifierKey::Control));
-		UI_COMMAND(LocalNext, "LocalNext", "Go to next local difference", EUserInterfaceActionType::Button, FInputGesture(EKeys::F7));
-		UI_COMMAND(RemotePrevious, "RemotePrevious", "Go to previous remote difference", EUserInterfaceActionType::Button, FInputGesture(EKeys::F8, EModifierKey::Control));
-		UI_COMMAND(RemoteNext, "RemoteNext", "Go to next remote difference", EUserInterfaceActionType::Button, FInputGesture(EKeys::F8));
-	}
-};
-
 SBlueprintMerge::SBlueprintMerge()
 	: Data()
 	, GraphView()
@@ -45,21 +21,16 @@ SBlueprintMerge::SBlueprintMerge()
 
 void SBlueprintMerge::Construct(const FArguments InArgs, const FBlueprintMergeData& InData)
 {
-	FMergeDifferencesListCommands::Register();
-
 	check( InData.OwningEditor.Pin().IsValid() );
 
 	Data = InData;
 
 	GraphView = SNew( SMergeGraphView, InData );
+	ComponentsView = SNew(SBorder);
+	DefaultsView = SNew(SBorder);
 
-	this->ChildSlot
-		[
-			GraphView.ToSharedRef()
-		];
-
-	/*ComponentsView;
-	DefaultsView;*/
+	InData.OwningEditor.Pin()->OnModeSet().AddSP( this, &SBlueprintMerge::OnModeChanged );
+	OnModeChanged( InData.OwningEditor.Pin()->GetCurrentMode() );
 }
 
 template< typename T >
@@ -126,4 +97,35 @@ FReply SBlueprintMerge::OnCancelClicked()
 	return FReply::Handled();
 }
 
+void SBlueprintMerge::OnModeChanged(FName NewMode)
+{
+	if (NewMode == FBlueprintEditorApplicationModes::StandardBlueprintEditorMode ||
+		NewMode == FBlueprintEditorApplicationModes::BlueprintMacroMode)
+	{
+		this->ChildSlot
+			[
+				GraphView.ToSharedRef()
+			];
+	}
+	else if (NewMode == FBlueprintEditorApplicationModes::BlueprintDefaultsMode)
+	{
+		this->ChildSlot
+			[
+				DefaultsView.ToSharedRef()
+			];
+	}
+	else if (NewMode == FBlueprintEditorApplicationModes::BlueprintComponentsMode ||
+			 NewMode == FBlueprintEditorApplicationModes::BlueprintInterfaceMode)
+	{
+		this->ChildSlot
+			[
+				ComponentsView.ToSharedRef()
+			];
+	}
+	else
+	{
+		ensureMsgf(false, TEXT("Diff panel does not support mode %s"), *NewMode.ToString());
+	}
+}
+
 #undef LOCTEXT_NAMESPACE
diff --git a/Engine/Source/Developer/Merge/Private/SBlueprintMerge.h b/Engine/Source/Developer/Merge/Private/SBlueprintMerge.h
--- a/Engine/Source/Developer/Merge/Private/SBlueprintMerge.h
+++ b/Engine/Source/Developer/Merge/Private/SBlueprintMerge.h
@@ -16,13 +16,14 @@ public:
 
 	void Construct(const FArguments InArgs, const FBlueprintMergeData& InData);
 
-protected:	
+private:
 	/** Helper functions */
 	UBlueprint* GetTargetBlueprint();
 
 	/** Event handlers */
 	FReply OnAcceptResultClicked();
 	FReply OnCancelClicked();
+	void OnModeChanged(FName NewMode);
 
 	FBlueprintMergeData Data;
 
diff --git a/Engine/Source/Editor/Kismet/Private/BlueprintEditor.cpp b/Engine/Source/Editor/Kismet/Private/BlueprintEditor.cpp
--- a/Engine/Source/Editor/Kismet/Private/BlueprintEditor.cpp
+++ b/Engine/Source/Editor/Kismet/Private/BlueprintEditor.cpp
@@ -1120,6 +1120,12 @@ void FBlueprintEditor::RegisterTabSpawners(const TSharedRef<class FTabManager>&
 	FWorkflowCentricApplication::RegisterTabSpawners(TabManager);
 }
 
+void FBlueprintEditor::SetCurrentMode(FName NewMode)
+{
+	OnModeSetData.Broadcast( NewMode );
+	FWorkflowCentricApplication::SetCurrentMode(NewMode);
+}
+
 void FBlueprintEditor::InitBlueprintEditor(const EToolkitMode::Type Mode, const TSharedPtr< IToolkitHost >& InitToolkitHost, const TArray<UBlueprint*>& InBlueprints, bool bShouldOpenInDefaultsMode)
 {
 	check(InBlueprints.Num() == 1 || bShouldOpenInDefaultsMode);
diff --git a/Engine/Source/Editor/Kismet/Public/BlueprintEditor.h b/Engine/Source/Editor/Kismet/Public/BlueprintEditor.h
--- a/Engine/Source/Editor/Kismet/Public/BlueprintEditor.h
+++ b/Engine/Source/Editor/Kismet/Public/BlueprintEditor.h
@@ -110,6 +110,10 @@ public:
 	virtual void RegisterTabSpawners(const TSharedRef<class FTabManager>& TabManager) override;
 	// End of IToolkit interface
 
+	DECLARE_MULTICAST_DELEGATE_OneParam(FOnModeSet, FName);
+	FOnModeSet& OnModeSet() { return OnModeSetData; }
+	virtual void SetCurrentMode(FName NewMode) override;
+
 public:
 	/**
 	 * Edits the specified blueprint
@@ -1019,6 +1023,9 @@ private:
 	/** Whether the current project is C++ or blueprint based */
 	bool bCodeBasedProject;
 
+	/** Delegates that are fired when the blueprint editor changes modes */
+	FOnModeSet OnModeSetData;
+
 };
 
 #undef LOCTEXT_NAMESPACE
diff --git a/Engine/Source/Editor/Kismet/Public/WorkflowOrientedApp/WorkflowCentricApplication.h b/Engine/Source/Editor/Kismet/Public/WorkflowOrientedApp/WorkflowCentricApplication.h
--- a/Engine/Source/Editor/Kismet/Public/WorkflowOrientedApp/WorkflowCentricApplication.h
+++ b/Engine/Source/Editor/Kismet/Public/WorkflowOrientedApp/WorkflowCentricApplication.h
@@ -34,7 +34,7 @@ public:
 	bool IsModeCurrent(FName ModeToCheck) const {return GetCurrentMode() == ModeToCheck;}
 
 	// Attempt to set the current mode.  If this mode is illegal or unknown, the mode will remain unchanged.
-	void SetCurrentMode(FName NewMode);
+	virtual void SetCurrentMode(FName NewMode);
 
 	void PushTabFactories(FWorkflowAllowedTabSet& FactorySetToPush);
 
