#UE4 fix for not all fields being hidden in details view of diff panel when 'show differing' is selected. Made the 'differing' filter a whitelist, rather than a blacklist. There are still bugs involving diffing of structures that have aggregate structures. Possibly also arrays and member objects depending on what the details view does with them.

diff --git a/Engine/Source/Developer/Merge/Private/SMergeDetailsView.cpp b/Engine/Source/Developer/Merge/Private/SMergeDetailsView.cpp
--- a/Engine/Source/Developer/Merge/Private/SMergeDetailsView.cpp
+++ b/Engine/Source/Developer/Merge/Private/SMergeDetailsView.cpp
@@ -22,9 +22,9 @@ void SMergeDetailsView::Construct(const FArguments InArgs, const FBlueprintMerge
 	TArray<FName> LocalIdenticalProperties, LocalDifferingProperties;
 	DiffUtils::CompareUnrelatedObjects(BaseCDO, BasePropertyMap, LocalCDO, LocalPropertyMap, LocalIdenticalProperties, LocalDifferingProperties);
 
-	TSet<FName> RemoteIdenticalPropertiesSet(RemoteIdenticalProperties);
-	TSet<FName> LocalIdenticalPropertiesSet(LocalIdenticalProperties);
-	TSet<FName> BaseIdenticalPropertiesSet(RemoteIdenticalPropertiesSet.Intersect(LocalIdenticalPropertiesSet));
+	TSet<FName> RemoteDifferingPropertiesSet(RemoteDifferingProperties);
+	TSet<FName> LocalDifferingPropertiesSet(LocalDifferingProperties);
+	TSet<FName> BaseDifferingPropertiesSet(RemoteDifferingPropertiesSet.Union(LocalDifferingPropertiesSet));
 	
 	/*	
 		DifferingProperties is an ordered list of all differing properties (properties added, removed or changed) in remote or local.
@@ -35,14 +35,14 @@ void SMergeDetailsView::Construct(const FArguments InArgs, const FBlueprintMerge
 			2. Iterate properties in remote, add any new properties to DifferingProperties
 			3. Iterate properties in local, add any new properties to DifferingProperties, if they have not already been added
 	*/
-	const auto AddPropertiesOrdered = []( UProperty const* Property, const TSet<FName>& IdenticalProperties, TArray<FName>& ResultingProperties )
+	const auto AddPropertiesOrdered = [](UProperty const* Property, const TSet<FName>& DifferingProperties, TArray<FName>& ResultingProperties)
 	{
 		if ((!Property->HasAnyPropertyFlags(CPF_Parm) && Property->HasAnyPropertyFlags(CPF_Edit)))
 		{
 			// contains check here is O(n), so we're needlessly n^2:
 			if (!ResultingProperties.Contains(Property->GetFName()))
 			{
-				if (!IdenticalProperties.Contains(Property->GetFName()))
+				if (DifferingProperties.Contains(Property->GetFName()))
 				{
 					ResultingProperties.Add(Property->GetFName());
 				}
@@ -55,7 +55,7 @@ void SMergeDetailsView::Construct(const FArguments InArgs, const FBlueprintMerge
 		const UClass* Class = BaseCDO->GetClass();
 		for (TFieldIterator<UProperty> PropertyIt(Class, EFieldIteratorFlags::IncludeSuper); PropertyIt; ++PropertyIt)
 		{
-			AddPropertiesOrdered( *PropertyIt, BaseIdenticalPropertiesSet, DifferingProperties );
+			AddPropertiesOrdered(*PropertyIt, BaseDifferingPropertiesSet, DifferingProperties);
 		}
 	}
 	if( RemoteCDO )
@@ -63,7 +63,7 @@ void SMergeDetailsView::Construct(const FArguments InArgs, const FBlueprintMerge
 		const UClass* Class = RemoteCDO->GetClass();
 		for (TFieldIterator<UProperty> PropertyIt(Class, EFieldIteratorFlags::IncludeSuper); PropertyIt; ++PropertyIt)
 		{
-			AddPropertiesOrdered(*PropertyIt, RemoteIdenticalPropertiesSet, DifferingProperties);
+			AddPropertiesOrdered(*PropertyIt, RemoteDifferingPropertiesSet, DifferingProperties);
 		}
 	}
 	if (LocalCDO)
@@ -71,7 +71,7 @@ void SMergeDetailsView::Construct(const FArguments InArgs, const FBlueprintMerge
 		const UClass* Class = LocalCDO->GetClass();
 		for (TFieldIterator<UProperty> PropertyIt(Class, EFieldIteratorFlags::IncludeSuper); PropertyIt; ++PropertyIt)
 		{
-			AddPropertiesOrdered(*PropertyIt, LocalIdenticalPropertiesSet, DifferingProperties);
+			AddPropertiesOrdered(*PropertyIt, LocalDifferingPropertiesSet, DifferingProperties);
 		}
 	}
 	CurrentDifference = -1;
@@ -79,19 +79,19 @@ void SMergeDetailsView::Construct(const FArguments InArgs, const FBlueprintMerge
 	// EMergeParticipant::MERGE_PARTICIPANT_REMOTE
 	{
 		DetailsViews.Add(
-			FDetailsDiff( RemoteCDO, RemotePropertyMap, RemoteIdenticalPropertiesSet)
+			FDetailsDiff(RemoteCDO, RemotePropertyMap, RemoteDifferingPropertiesSet)
 		);
 	}
 	// EMergeParticipant::MERGE_PARTICIPANT_BASE
 	{
 		DetailsViews.Add(
-			FDetailsDiff(BaseCDO, BasePropertyMap, BaseIdenticalPropertiesSet)
+			FDetailsDiff(BaseCDO, BasePropertyMap, BaseDifferingPropertiesSet)
 		);
 	}
 	// EMergeParticipant::MERGE_PARTICIPANT_LOCAL
 	{
 		DetailsViews.Add(
-			FDetailsDiff(LocalCDO, LocalPropertyMap, LocalIdenticalPropertiesSet)
+			FDetailsDiff(LocalCDO, LocalPropertyMap, LocalDifferingPropertiesSet)
 		);
 	}
 
diff --git a/Engine/Source/Editor/Kismet/Private/DetailsDiff.cpp b/Engine/Source/Editor/Kismet/Private/DetailsDiff.cpp
--- a/Engine/Source/Editor/Kismet/Private/DetailsDiff.cpp
+++ b/Engine/Source/Editor/Kismet/Private/DetailsDiff.cpp
@@ -5,7 +5,7 @@
 #include "IDetailsView.h"
 #include "PropertyEditorModule.h"
 
-FDetailsDiff::FDetailsDiff(const UObject* InObject, const TMap< FName, const UProperty* >& InPropertyMap, const TSet< FName >& InIdenticalProperties)
+FDetailsDiff::FDetailsDiff(const UObject* InObject, const TMap< FName, const UProperty* >& InPropertyMap, const TSet< FName >& InDifferingProperties)
 	: PropertyMap( InPropertyMap )
 	, DetailsView()
 {
@@ -17,7 +17,7 @@ FDetailsDiff::FDetailsDiff(const UObject* InObject, const TMap< FName, const UPr
 	DetailsView->SetIsPropertyEditingEnabledDelegate(FIsPropertyEditingEnabled::CreateStatic([]{return false; }));
 	// This is a read only details view (see the property editing delegate), but it is not const correct:
 	DetailsView->SetObject(const_cast<UObject*>(InObject));
-	DetailsView->UpdateIdenticalProperties(InIdenticalProperties);
+	DetailsView->UpdateDifferingProperties(InDifferingProperties);
 }
 
 void FDetailsDiff::HighlightProperty(FName PropertyName)
diff --git a/Engine/Source/Editor/Kismet/Private/SBlueprintDiff.cpp b/Engine/Source/Editor/Kismet/Private/SBlueprintDiff.cpp
--- a/Engine/Source/Editor/Kismet/Private/SBlueprintDiff.cpp
+++ b/Engine/Source/Editor/Kismet/Private/SBlueprintDiff.cpp
@@ -33,7 +33,6 @@ public:
 					, FNamePropertyMap InOldProperties
 					, const UObject* InNewCDO
 					, FNamePropertyMap InNewProperities
-					, TSet<FName> InIdenticalProperties
 					, TArray<FName> InDifferingProperties );
 
 	TSharedRef<SWidget> OldDetailsWidget() { return OldDetails.DetailsWidget(); }
@@ -58,10 +57,9 @@ FCDODiffControl::FCDODiffControl(
 		, FNamePropertyMap InOldProperties
 		, const UObject* InNewCDO
 		, FNamePropertyMap InNewProperities
-		, TSet<FName> InIdenticalProperties
 		, TArray<FName> InDifferingProperties )
-	: OldDetails( InOldCDO, InOldProperties, InIdenticalProperties )
-	, NewDetails( InNewCDO, InNewProperities, InIdenticalProperties )
+	: OldDetails(InOldCDO, InOldProperties, TSet<FName>(InDifferingProperties))
+	, NewDetails(InNewCDO, InNewProperities, TSet<FName>(InDifferingProperties))
 	, DifferingProperties( InDifferingProperties )
 	, CurrentDifference( -1 )
 {
@@ -570,7 +568,8 @@ void DiffUtils::CompareUnrelatedObjects( UObject const* A, const TMap< FName, co
 			{
 				const void* AValue = (*AProp)->ContainerPtrToValuePtr<void>(A);
 				const void* BValue = (*BProp)->ContainerPtrToValuePtr<void>(B);
-				if ((*AProp)->Identical(AValue, BValue, 0))
+
+				if ((*AProp)->Identical(AValue, BValue, PPF_DeepComparison))
 				{
 					OutIdenticalProperties.Add(PropertyName);
 					continue;
@@ -1221,7 +1220,7 @@ TSharedRef<SWidget> SBlueprintDiff::GenerateDefaultsPanel()
 	TArray<FName> IdenticalProperties, DifferingProperties;
 	DiffUtils::CompareUnrelatedObjects( A, PropertyMapA, B, PropertyMapB, IdenticalProperties, DifferingProperties);
 
-	auto NewDiffControl = TSharedPtr<FCDODiffControl>(new FCDODiffControl(A, PropertyMapA, B, PropertyMapB, TSet<FName>(IdenticalProperties), DifferingProperties) );
+	auto NewDiffControl = TSharedPtr<FCDODiffControl>(new FCDODiffControl(A, PropertyMapA, B, PropertyMapB, DifferingProperties) );
 	//Splitter for left and right blueprint. Current convention is for the local (probably newer?) blueprint to be on the right:
 	DiffControl = NewDiffControl;
 	return SNew(SSplitter)
diff --git a/Engine/Source/Editor/Kismet/Public/DetailsDiff.h b/Engine/Source/Editor/Kismet/Public/DetailsDiff.h
--- a/Engine/Source/Editor/Kismet/Public/DetailsDiff.h
+++ b/Engine/Source/Editor/Kismet/Public/DetailsDiff.h
@@ -4,7 +4,7 @@
 class KISMET_API FDetailsDiff
 {
 public:
-	FDetailsDiff( const UObject* InObject, const TMap< FName, const UProperty* >& InPropertyMap, const TSet< FName >& InIdenticalProperties );
+	FDetailsDiff( const UObject* InObject, const TMap< FName, const UProperty* >& InPropertyMap, const TSet< FName >& InDifferingProperties );
 
 	void HighlightProperty( FName PropertyName );
 	TSharedRef< SWidget > DetailsWidget();
diff --git a/Engine/Source/Editor/PropertyEditor/Private/DetailItemNode.cpp b/Engine/Source/Editor/PropertyEditor/Private/DetailItemNode.cpp
--- a/Engine/Source/Editor/PropertyEditor/Private/DetailItemNode.cpp
+++ b/Engine/Source/Editor/PropertyEditor/Private/DetailItemNode.cpp
@@ -282,7 +282,7 @@ static bool PassesAllFilters( const FDetailLayoutCustomization& InCustomization,
 
 			const bool bPassesSearchFilter = bSearchFilterIsEmpty || ( bIsNotBeingFiltered || bIsSeenDueToFiltering || bIsParentSeenDueToFiltering );
 			const bool bPassesModifiedFilter = bPassesSearchFilter && ( InFilter.bShowOnlyModifiedProperties == false || PropertyNodePin->GetDiffersFromDefault() == true );
-			const bool bPassesDifferingFilter = InFilter.bShowOnlyDiffering == false || InFilter.NamesOfMatchingProperties.Find( PropertyNodePin->GetProperty()->GetFName() ) == NULL;
+			const bool bPassesDifferingFilter = InFilter.bShowOnlyDiffering == false || InFilter.DifferingProperties.Find(PropertyNodePin->GetProperty()->GetFName()) != NULL;
 
 			// The property node is visible (note categories are never visible unless they have a child that is visible )
 			bPassesAllFilters = bPassesSearchFilter && bPassesModifiedFilter && bPassesDifferingFilter;
diff --git a/Engine/Source/Editor/PropertyEditor/Private/SDetailsViewBase.h b/Engine/Source/Editor/PropertyEditor/Private/SDetailsViewBase.h
--- a/Engine/Source/Editor/PropertyEditor/Private/SDetailsViewBase.h
+++ b/Engine/Source/Editor/PropertyEditor/Private/SDetailsViewBase.h
@@ -66,7 +66,7 @@ struct FDetailFilter
 	bool bShowAllAdvanced;
 	/** If we should only show differing properties */
 	bool bShowOnlyDiffering;
-	TSet<FName> NamesOfMatchingProperties;
+	TSet<FName> DifferingProperties;
 };
 
 struct FDetailColumnSizeData
@@ -352,8 +352,8 @@ protected:
 	/** Called when the filter text changes.  This filters specific property nodes out of view */
 	void OnFilterTextChanged(const FText& InFilterText);
 
-	// @todo doc where does this belong?
-	virtual void UpdateIdenticalProperties(const TSet<FName> IdenticalProperties) override { CurrentFilter.NamesOfMatchingProperties = IdenticalProperties; }
+	/** Called when the list of currently differing properties changes */
+	virtual void UpdateDifferingProperties(const TSet<FName> InDifferingProperties) override { CurrentFilter.DifferingProperties = InDifferingProperties; }
 
 	/** 
 	 * Hides or shows properties based on the passed in filter text
diff --git a/Engine/Source/Editor/PropertyEditor/Public/IDetailsView.h b/Engine/Source/Editor/PropertyEditor/Public/IDetailsView.h
--- a/Engine/Source/Editor/PropertyEditor/Public/IDetailsView.h
+++ b/Engine/Source/Editor/PropertyEditor/Public/IDetailsView.h
@@ -175,7 +175,7 @@ public:
 	virtual void HighlightProperty(const UProperty* Property) = 0;
 
 	/**
-	 * Sets the set of properties that are considered identical, used when filtering out identical properties (showing only differing properties)
+	 * Sets the set of properties that are considered differing, used when filtering out identical properties
 	 */
-	virtual void UpdateIdenticalProperties(const TSet<FName> IdenticalProperties) = 0;
+	virtual void UpdateDifferingProperties(const TSet<FName> DifferingProperties) = 0;
 };
