Merge view for components no longer lists the same difference multiple times

diff --git a/Engine/Source/Developer/Merge/Private/SMergeTreeView.cpp b/Engine/Source/Developer/Merge/Private/SMergeTreeView.cpp
--- a/Engine/Source/Developer/Merge/Private/SMergeTreeView.cpp
+++ b/Engine/Source/Developer/Merge/Private/SMergeTreeView.cpp
@@ -27,25 +27,32 @@ void SMergeTreeView::Construct(const FArguments InArgs, const FBlueprintMergeDat
 	{
 		const bool bStillHaveLocalDiffs = LocalDifferingProperties.IsValidIndex(LocalIter);
 		const bool bStillHaveRemoteDiffs = RemoteDifferingProperties.IsValidIndex(RemoteIter);
+		FSCSDiffEntry Entry;
+
 		if (bStillHaveLocalDiffs && bStillHaveRemoteDiffs)
 		{
 			if (LocalDiffsSortKeys[LocalIter] <= RemoteDiffsSortKeys[RemoteIter])
 			{
-				DifferingProperties.Push(LocalDifferingProperties[LocalIter++]);
+				Entry = LocalDifferingProperties[LocalIter++];
 			}
 			else
 			{
-				DifferingProperties.Push(RemoteDifferingProperties[RemoteIter++]);
+				Entry = RemoteDifferingProperties[RemoteIter++];
 			}
 		}
 		else if (bStillHaveLocalDiffs)
 		{
-			DifferingProperties.Push(LocalDifferingProperties[LocalIter++]);
+			Entry = LocalDifferingProperties[LocalIter++];
 		}
 		else
 		{
 			check(bStillHaveRemoteDiffs);
-			DifferingProperties.Push(RemoteDifferingProperties[RemoteIter++]);
+			Entry = RemoteDifferingProperties[RemoteIter++];
+		}
+
+		if (!DifferingProperties.Contains(Entry))
+		{
+			DifferingProperties.Push(Entry);
 		}
 	}
 
diff --git a/Engine/Source/Editor/Kismet/Public/SBlueprintDiff.h b/Engine/Source/Editor/Kismet/Public/SBlueprintDiff.h
--- a/Engine/Source/Editor/Kismet/Public/SBlueprintDiff.h
+++ b/Engine/Source/Editor/Kismet/Public/SBlueprintDiff.h
@@ -36,6 +36,11 @@ struct FSCSDiffEntry
 	FName PropertyName;
 };
 
+inline bool operator ==(const FSCSDiffEntry A, const FSCSDiffEntry B)
+{
+	return A.TreeNodeName == B.TreeNodeName && A.PropertyName == B.PropertyName;
+}
+
 namespace DiffUtils
 {
 	KISMET_API TMap< FName, const UProperty* > GetProperties( const UObject* ForObj );
