#UE4 spawning the merge tab in the blueprint editor's tab manager, this makes the merge tool window appear in the blueprint editor initially, much more intuitive because the merge tool responds to tab switches performed in the blueprint editor...

diff --git a/Engine/Source/Developer/Merge/Private/Merge.cpp b/Engine/Source/Developer/Merge/Private/Merge.cpp
--- a/Engine/Source/Developer/Merge/Private/Merge.cpp
+++ b/Engine/Source/Developer/Merge/Private/Merge.cpp
@@ -226,7 +226,7 @@ TSharedRef<SDockTab> FMerge::GenerateMergeWidget(const UBlueprint& Object, TShar
 		}
 	}
 
-	TSharedRef<SDockTab> Tab =  FGlobalTabmanager::Get()->InvokeTab(MergeToolTabId);
+	TSharedRef<SDockTab> Tab = Editor->GetTabManager()->InvokeTab(MergeToolTabId);
 	Tab->SetContent(Contents.ToSharedRef());
 	ActiveTab = Tab;
 	return Tab;
