Backing out my last change... it results in the merge tool not staying open when tabs are switched. This restores the previous bug: the merge tool is initially floating

diff --git a/Engine/Source/Developer/Merge/Private/Merge.cpp b/Engine/Source/Developer/Merge/Private/Merge.cpp
--- a/Engine/Source/Developer/Merge/Private/Merge.cpp
+++ b/Engine/Source/Developer/Merge/Private/Merge.cpp
@@ -226,7 +226,7 @@ TSharedRef<SDockTab> FMerge::GenerateMergeWidget(const UBlueprint& Object, TShar
 		}
 	}
 
-	TSharedRef<SDockTab> Tab = Editor->GetTabManager()->InvokeTab(MergeToolTabId);
+	TSharedRef<SDockTab> Tab =  FGlobalTabmanager::Get()->InvokeTab(MergeToolTabId);
 	Tab->SetContent(Contents.ToSharedRef());
 	ActiveTab = Tab;
 	return Tab;
