Files are now backed up when merge begins, merge window correctly docks with blueprint editor by default

diff --git a/Engine/Source/Developer/Merge/Private/SBlueprintMerge.cpp b/Engine/Source/Developer/Merge/Private/SBlueprintMerge.cpp
--- a/Engine/Source/Developer/Merge/Private/SBlueprintMerge.cpp
+++ b/Engine/Source/Developer/Merge/Private/SBlueprintMerge.cpp
@@ -20,12 +20,32 @@ SBlueprintMerge::SBlueprintMerge()
 {
 }
 
+static void WriteBackup(UPackage& Package, const FString& Directory, const FString& Filename)
+{
+	if (GEditor)
+	{
+		const FString DestinationFilename = Directory + FString("/") + Filename;
+		FString OriginalFilename;
+		if (FPackageName::DoesPackageExist(Package.GetName(), NULL, &OriginalFilename))
+		{
+			IFileManager::Get().Copy(*DestinationFilename, *OriginalFilename);
+		}
+	}
+}
+
+static const FString BackupSubDir = FPaths::GameSavedDir() / TEXT("Backup") / TEXT("Resolve_Backup[") + FDateTime::Now().ToString(TEXT("%Y-%m-%d-%H-%M-%S")) + TEXT("]");
+
 void SBlueprintMerge::Construct(const FArguments InArgs, const FBlueprintMergeData& InData)
 {
 	check( InData.OwningEditor.Pin().IsValid() );
 
 	Data = InData;
 
+	// Because merge operations are so destructive and can be confusing, lets write backups of the files involved:
+	WriteBackup(*Data.BlueprintRemote->GetOutermost(), BackupSubDir, TEXT("RemoteAsset") + FPackageName::GetAssetPackageExtension());
+	WriteBackup(*Data.BlueprintBase->GetOutermost(), BackupSubDir, TEXT("CommonBaseAsset") + FPackageName::GetAssetPackageExtension());
+	WriteBackup(*Data.BlueprintLocal->GetOutermost(), BackupSubDir, TEXT("LocalAsset") + FPackageName::GetAssetPackageExtension());
+
 	auto GraphView = SNew( SMergeGraphView, InData );
 	GraphControl.Widget = GraphView;
 	GraphControl.DiffControl = &GraphView.Get();
@@ -92,29 +112,6 @@ void SBlueprintMerge::Construct(const FArguments InArgs, const FBlueprintMergeDa
 	OnModeChanged( InData.OwningEditor.Pin()->GetCurrentMode() );
 }
 
-template< typename T >
-static void CopyToShared( const TArray<T>& Source, TArray< TSharedPtr<T> >& Dest )
-{
-	check( Dest.Num() == 0 );
-	for( auto i : Source )
-	{
-		Dest.Add( MakeShareable(new T( i )) );
-	}
-}
-
-static void WriteBackup( UPackage& Package, const FString& Directory, const FString& Filename )
-{
-	if( GEditor )
-	{
-		const FString DestinationFilename = Directory + FString("/") + Filename;
-		FString OriginalFilename;
-		if( FPackageName::DoesPackageExist(Package.GetName(), NULL, &OriginalFilename) )
-		{
-			IFileManager::Get().Copy(*DestinationFilename, *OriginalFilename);
-		}
-	}
-}
-
 UBlueprint* SBlueprintMerge::GetTargetBlueprint()
 {
 	return Data.OwningEditor.Pin()->GetBlueprintObj();
@@ -142,12 +139,6 @@ bool SBlueprintMerge::HasPrevDiff() const
 
 void SBlueprintMerge::OnAcceptResultClicked()
 {
-	// Because merge operations are so destructive and can be confusing, lets write backups of the files involved:
-	const FString BackupSubDir = FPaths::GameSavedDir() / TEXT("Backup") / TEXT("Resolve_Backup[") + FDateTime::Now().ToString(TEXT("%Y-%m-%d-%H-%M-%S")) + TEXT("]");
-	WriteBackup(*Data.BlueprintRemote->GetOutermost(), BackupSubDir, TEXT("RemoteAsset") + FPackageName::GetAssetPackageExtension() );
-	WriteBackup(*Data.BlueprintBase->GetOutermost(), BackupSubDir, TEXT("CommonBaseAsset") + FPackageName::GetAssetPackageExtension() );
-	WriteBackup(*Data.BlueprintLocal->GetOutermost(), BackupSubDir, TEXT("LocalAsset") + FPackageName::GetAssetPackageExtension());
-
 	UPackage* Package = GetTargetBlueprint()->GetOutermost();
 	TArray<UPackage*> PackagesToSave;
 	PackagesToSave.Add(Package);
diff --git a/Engine/Source/Editor/Kismet/Private/BlueprintEditorModes.cpp b/Engine/Source/Editor/Kismet/Private/BlueprintEditorModes.cpp
--- a/Engine/Source/Editor/Kismet/Private/BlueprintEditorModes.cpp
+++ b/Engine/Source/Editor/Kismet/Private/BlueprintEditorModes.cpp
@@ -121,16 +121,6 @@ FBlueprintEditorApplicationMode::FBlueprintEditorApplicationMode(TSharedPtr<clas
 					)
 				)
 			)
-			->Split
-			(
-				FTabManager::NewSplitter() ->SetOrientation( Orient_Horizontal )
-				->SetSizeCoefficient(0.35f)
-				->Split
-				(
-					FTabManager::NewStack()
-					->AddTab( FBlueprintEditorTabs::MergeToolID, ETabState::ClosedTab )
-				)
-			)
 		);
 	
 	// setup toolbar
@@ -203,16 +193,6 @@ FBlueprintDefaultsApplicationMode::FBlueprintDefaultsApplicationMode(TSharedPtr<
 				->SetHideTabWell(true)
 				->AddTab( FBlueprintEditorTabs::DefaultEditorID, ETabState::OpenedTab )
 			)
-			->Split
-			(
-				FTabManager::NewSplitter() ->SetOrientation( Orient_Horizontal )
-				->SetSizeCoefficient(0.35f)
-				->Split
-				(
-					FTabManager::NewStack()
-					->AddTab( FBlueprintEditorTabs::MergeToolID, ETabState::ClosedTab )
-				)
-			)
 		);
 
 	// setup toolbar
@@ -295,16 +275,6 @@ FBlueprintComponentsApplicationMode::FBlueprintComponentsApplicationMode(TShared
 					->AddTab( FBlueprintEditorTabs::SCSViewportID, ETabState::OpenedTab )
 				)
 			)
-			->Split
-			(
-				FTabManager::NewSplitter() ->SetOrientation( Orient_Horizontal )
-				->SetSizeCoefficient(0.35f)
-				->Split
-				(
-					FTabManager::NewStack()
-					->AddTab( FBlueprintEditorTabs::MergeToolID, ETabState::ClosedTab )
-				)
-			)
 		);
 
 	// setup toolbar
diff --git a/Engine/Source/Editor/Kismet/Private/BlueprintEditorTabs.cpp b/Engine/Source/Editor/Kismet/Private/BlueprintEditorTabs.cpp
--- a/Engine/Source/Editor/Kismet/Private/BlueprintEditorTabs.cpp
+++ b/Engine/Source/Editor/Kismet/Private/BlueprintEditorTabs.cpp
@@ -9,7 +9,6 @@ const FName FBlueprintEditorTabs::DebugID( TEXT( "Debug" ) );
 const FName FBlueprintEditorTabs::PaletteID( TEXT( "PaletteList" ) );
 const FName FBlueprintEditorTabs::CompilerResultsID( TEXT( "CompilerResults" ) );
 const FName FBlueprintEditorTabs::FindResultsID( TEXT( "FindResults" ) );
-const FName FBlueprintEditorTabs::MergeToolID(TEXT("MergeTool"));
 const FName FBlueprintEditorTabs::ConstructionScriptEditorID( TEXT( "ConstructionScriptEditor" ) );
 const FName FBlueprintEditorTabs::SCSViewportID( TEXT( "SCSViewport" ) );
 const FName FBlueprintEditorTabs::MyBlueprintID( TEXT( "MyBlueprint" ) );
diff --git a/Engine/Source/Editor/Kismet/Public/BlueprintEditorTabs.h b/Engine/Source/Editor/Kismet/Public/BlueprintEditorTabs.h
--- a/Engine/Source/Editor/Kismet/Public/BlueprintEditorTabs.h
+++ b/Engine/Source/Editor/Kismet/Public/BlueprintEditorTabs.h
@@ -11,7 +11,6 @@ struct KISMET_API FBlueprintEditorTabs
 	static const FName PaletteID;
 	static const FName CompilerResultsID;
 	static const FName FindResultsID;
-	static const FName MergeToolID;
 	static const FName ConstructionScriptEditorID;
 	static const FName SCSViewportID;
 	static const FName MyBlueprintID;
diff --git a/Engine/Source/Editor/MainFrame/Private/MainFrameModule.cpp b/Engine/Source/Editor/MainFrame/Private/MainFrameModule.cpp
--- a/Engine/Source/Editor/MainFrame/Private/MainFrameModule.cpp
+++ b/Engine/Source/Editor/MainFrame/Private/MainFrameModule.cpp
@@ -143,12 +143,19 @@ void FMainFrameModule::CreateDefaultMainFrame( const bool bStartImmersivePIE )
 				->AddArea
 				(
 					FTabManager::NewArea(WindowSize)
+					->SetOrientation(Orient_Vertical)
 					->Split
 					(
 						FTabManager::NewStack()
 						->SetSizeCoefficient(1.0f)
 						->AddTab("StandaloneToolkit", ETabState::ClosedTab)
 					)
+					->Split
+					(
+						FTabManager::NewStack()
+						->SetSizeCoefficient(0.35f)
+						->AddTab("MergeTool", ETabState::ClosedTab)
+					)
 				)
 			);
 
