Abusing a global, moved to member function. This triggered an assert on Linux, may also not play nice with monolithic builds that include the merge module...

diff --git a/Engine/Source/Developer/Merge/Private/SBlueprintMerge.cpp b/Engine/Source/Developer/Merge/Private/SBlueprintMerge.cpp
--- a/Engine/Source/Developer/Merge/Private/SBlueprintMerge.cpp
+++ b/Engine/Source/Developer/Merge/Private/SBlueprintMerge.cpp
@@ -33,14 +33,14 @@ static void WriteBackup(UPackage& Package, const FString& Directory, const FStri
 	}
 }
 
-static const FString BackupSubDir = FPaths::GameSavedDir() / TEXT("Backup") / TEXT("Resolve_Backup[") + FDateTime::Now().ToString(TEXT("%Y-%m-%d-%H-%M-%S")) + TEXT("]");
-
 void SBlueprintMerge::Construct(const FArguments InArgs, const FBlueprintMergeData& InData)
 {
 	check( InData.OwningEditor.Pin().IsValid() );
 
 	Data = InData;
 
+	BackupSubDir = FPaths::GameSavedDir() / TEXT("Backup") / TEXT("Resolve_Backup[") + FDateTime::Now().ToString(TEXT("%Y-%m-%d-%H-%M-%S")) + TEXT("]");
+
 	// Because merge operations are so destructive and can be confusing, lets write backups of the files involved:
 	WriteBackup(*Data.BlueprintRemote->GetOutermost(), BackupSubDir, TEXT("RemoteAsset") + FPackageName::GetAssetPackageExtension());
 	WriteBackup(*Data.BlueprintBase->GetOutermost(), BackupSubDir, TEXT("CommonBaseAsset") + FPackageName::GetAssetPackageExtension());
diff --git a/Engine/Source/Developer/Merge/Private/SBlueprintMerge.h b/Engine/Source/Developer/Merge/Private/SBlueprintMerge.h
--- a/Engine/Source/Developer/Merge/Private/SBlueprintMerge.h
+++ b/Engine/Source/Developer/Merge/Private/SBlueprintMerge.h
@@ -29,7 +29,8 @@ private:
 	void OnCancelClicked();
 	void OnModeChanged(FName NewMode);
 
-	FBlueprintMergeData Data;
+	FBlueprintMergeData		Data; 
+	FString					BackupSubDir;
 
 	TSharedPtr<SBorder>		MainView;
 
