Added mising #undef LOCTEXT_NAMESPACE in SMergeGraphView.cpp

diff --git a/Engine/Source/Developer/Merge/Private/SMergeGraphView.cpp b/Engine/Source/Developer/Merge/Private/SMergeGraphView.cpp
--- a/Engine/Source/Developer/Merge/Private/SMergeGraphView.cpp
+++ b/Engine/Source/Developer/Merge/Private/SMergeGraphView.cpp
@@ -487,3 +487,5 @@ const FSlateBrush* SMergeGraphView::GetLockViewImage() const
 {
 	return bViewsAreLocked ? FEditorStyle::GetBrush("GenericLock") : FEditorStyle::GetBrush("GenericUnlock");
 }
+
+#undef LOCTEXT_NAMESPACE
\ No newline at end of file
