Conflict highlighting for the graph view in the merge tool

diff --git a/Engine/Source/Developer/Merge/Private/BlueprintMergeData.h b/Engine/Source/Developer/Merge/Private/BlueprintMergeData.h
--- a/Engine/Source/Developer/Merge/Private/BlueprintMergeData.h
+++ b/Engine/Source/Developer/Merge/Private/BlueprintMergeData.h
@@ -53,3 +53,12 @@ namespace EMergeParticipant
 		MERGE_PARTICIPANT_LOCAL,
 	};
 }
+
+class IMergeControl
+{
+public:
+	virtual void HighlightNextConflict() = 0;
+	virtual void HighlightPrevConflict() = 0;
+	virtual bool HasNextConflict() const = 0;
+	virtual bool HasPrevConflict() const = 0;
+};
diff --git a/Engine/Source/Developer/Merge/Private/SBlueprintMerge.cpp b/Engine/Source/Developer/Merge/Private/SBlueprintMerge.cpp
--- a/Engine/Source/Developer/Merge/Private/SBlueprintMerge.cpp
+++ b/Engine/Source/Developer/Merge/Private/SBlueprintMerge.cpp
@@ -18,7 +18,8 @@
 
 SBlueprintMerge::SBlueprintMerge()
 	: Data()
-	, CurrentDiffControl( NULL )
+	, CurrentDiffControl( nullptr )
+	, CurrentMergeControl( nullptr )
 {
 }
 
@@ -28,7 +29,7 @@ static void WriteBackup(UPackage& Package, const FString& Directory, const FStri
 	{
 		const FString DestinationFilename = Directory + FString("/") + Filename;
 		FString OriginalFilename;
-		if (FPackageName::DoesPackageExist(Package.GetName(), NULL, &OriginalFilename))
+		if (FPackageName::DoesPackageExist(Package.GetName(), nullptr, &OriginalFilename))
 		{
 			IFileManager::Get().Copy(*DestinationFilename, *OriginalFilename);
 		}
@@ -51,6 +52,7 @@ void SBlueprintMerge::Construct(const FArguments InArgs, const FBlueprintMergeDa
 	auto GraphView = SNew(SMergeGraphView, InData);
 	GraphControl.Widget = GraphView;
 	GraphControl.DiffControl = &GraphView.Get();
+	GraphControl.MergeControl = &GraphView.Get();
 
 	auto TreeView = SNew(SMergeTreeView, InData);
 	TreeControl.Widget = TreeView;
@@ -75,6 +77,25 @@ void SBlueprintMerge::Construct(const FArguments InArgs, const FBlueprintMergeDa
 		, LOCTEXT("NextMergeTooltip", "Go to next difference")
 		, FSlateIcon(FEditorStyle::GetStyleSetName(), "BlueprintMerge.NextDiff") 
 	);
+
+	// conflict navigation buttons:
+	ToolbarBuilder.AddSeparator();
+	ToolbarBuilder.AddToolBarButton( 
+		FUIAction( FExecuteAction::CreateSP(this, &SBlueprintMerge::PrevConflict), FCanExecuteAction::CreateSP(this, &SBlueprintMerge::HasPrevConflict) )
+		, NAME_None
+		, LOCTEXT("PrevConflictLabel", "Prev Conflict")
+		, LOCTEXT("PrevConflictTooltip", "Go to previous conflict")
+		, FSlateIcon(FEditorStyle::GetStyleSetName(), "BlueprintMerge.PrevDiff")
+	);
+	ToolbarBuilder.AddToolBarButton(
+		FUIAction(FExecuteAction::CreateSP(this, &SBlueprintMerge::NextConflict), FCanExecuteAction::CreateSP(this, &SBlueprintMerge::HasNextConflict))
+		, NAME_None
+		, LOCTEXT("NextConflictLabel", "Next Conflict")
+		, LOCTEXT("NextConflictTooltip", "Go to next conflict")
+		, FSlateIcon(FEditorStyle::GetStyleSetName(), "BlueprintMerge.NextDiff") 
+	);
+
+	// buttons for finishing the merge:
 	ToolbarBuilder.AddSeparator();
 	ToolbarBuilder.AddToolBarButton(
 		FUIAction(FExecuteAction::CreateSP(this, &SBlueprintMerge::OnFinishMerge))
@@ -139,6 +160,26 @@ bool SBlueprintMerge::HasPrevDiff() const
 	return CurrentDiffControl && CurrentDiffControl->HasPrevDifference();
 }
 
+void SBlueprintMerge::NextConflict()
+{
+	CurrentMergeControl->HighlightNextConflict();
+}
+
+void SBlueprintMerge::PrevConflict()
+{
+	CurrentMergeControl->HighlightPrevConflict();
+}
+
+bool SBlueprintMerge::HasNextConflict() const
+{
+	return CurrentMergeControl && CurrentMergeControl->HasNextConflict();
+}
+
+bool SBlueprintMerge::HasPrevConflict() const
+{
+	return CurrentMergeControl && CurrentMergeControl->HasPrevConflict();
+}
+
 void SBlueprintMerge::OnFinishMerge()
 {
 	UPackage* Package = GetTargetBlueprint()->GetOutermost();
@@ -173,17 +214,20 @@ void SBlueprintMerge::OnModeChanged(FName NewMode)
 	{
 		MainView->SetContent( GraphControl.Widget.ToSharedRef() );
 		CurrentDiffControl = GraphControl.DiffControl;
+		CurrentMergeControl = GraphControl.MergeControl;
 	}
 	else if (NewMode == FBlueprintEditorApplicationModes::BlueprintComponentsMode)
 	{
 		MainView->SetContent(TreeControl.Widget.ToSharedRef());
 		CurrentDiffControl = TreeControl.DiffControl;
+		CurrentMergeControl = TreeControl.MergeControl;
 	}
 	else if (NewMode == FBlueprintEditorApplicationModes::BlueprintDefaultsMode ||
 			NewMode == FBlueprintEditorApplicationModes::BlueprintInterfaceMode)
 	{
 		MainView->SetContent(DetailsControl.Widget.ToSharedRef());
 		CurrentDiffControl = DetailsControl.DiffControl;
+		CurrentMergeControl = DetailsControl.MergeControl;
 	}
 	else
 	{
diff --git a/Engine/Source/Developer/Merge/Private/SBlueprintMerge.h b/Engine/Source/Developer/Merge/Private/SBlueprintMerge.h
--- a/Engine/Source/Developer/Merge/Private/SBlueprintMerge.h
+++ b/Engine/Source/Developer/Merge/Private/SBlueprintMerge.h
@@ -25,6 +25,12 @@ private:
 	void PrevDiff();
 	bool HasNextDiff() const;
 	bool HasPrevDiff() const;
+
+	void NextConflict();
+	void PrevConflict();
+	bool HasNextConflict() const;
+	bool HasPrevConflict() const;
+
 	void OnFinishMerge();
 	void OnCancelClicked();
 	void OnModeChanged(FName NewMode);
@@ -34,23 +40,26 @@ private:
 
 	TSharedPtr<SBorder>		MainView;
 
-	struct FDiffControl
+	struct FMergeControl
 	{
-		FDiffControl()
+		FMergeControl()
 			: Widget()
 			, DiffControl(NULL)
+			, MergeControl(NULL)
 		{
 		}
 
 		TSharedPtr<SWidget> Widget;
 		class IDiffControl* DiffControl;
+		class IMergeControl* MergeControl;
 	};
 
-	FDiffControl GraphControl;
-	FDiffControl TreeControl;
-	FDiffControl DetailsControl;
+	FMergeControl GraphControl;
+	FMergeControl TreeControl;
+	FMergeControl DetailsControl;
 
 	class IDiffControl* CurrentDiffControl;
+	class IMergeControl* CurrentMergeControl;
 
 	// This has to be allocated here because SListView cannot own the list
 	// that it is displaying. It also seems like the display list *has*
diff --git a/Engine/Source/Developer/Merge/Private/SMergeGraphView.cpp b/Engine/Source/Developer/Merge/Private/SMergeGraphView.cpp
--- a/Engine/Source/Developer/Merge/Private/SMergeGraphView.cpp
+++ b/Engine/Source/Developer/Merge/Private/SMergeGraphView.cpp
@@ -7,11 +7,25 @@
 
 #define LOCTEXT_NAMESPACE "SMergeGraphView"
 
+struct FGraphMergeConflict
+{
+	FGraphMergeConflict(FName InGraphName, const FDiffSingleResult& InFirstResult, const FDiffSingleResult& InSecondResult )
+		: GraphName(InGraphName)
+		, FirstResult(InFirstResult)
+		, SecondResult(InSecondResult)
+	{
+	}
+
+	FName GraphName;
+	FDiffSingleResult FirstResult;
+	FDiffSingleResult SecondResult;
+};
+
 // This cludge is used to move arrays of structures into arrays of 
 // shared ptrs of those structures. It is used because SListView
 // does not support arrays of values:
 template< typename T >
-TArray< TSharedPtr<T> > HackToShared(const TArray<T>& Values)
+TArray< TSharedPtr<T> > ConcreteToShared(const TArray<T>& Values)
 {
 	TArray< TSharedPtr<T> > Ret;
 	for (const auto& Value : Values)
@@ -35,7 +49,7 @@ struct FBlueprintRevPair
 
 struct FMergeGraphRowEntry
 {
-	FString GraphName;
+	FName GraphName;
 
 	// These lists contain shared ptrs because they are displayed by a
 	// SListView, and it currently does not support lists of values.
@@ -48,13 +62,13 @@ struct FMergeGraphRowEntry
 	bool bDiffersInLocal;
 };
 
-static UEdGraph* FindGraphByName(UBlueprint const& FromBlueprint, const FString& GraphName)
+static UEdGraph* FindGraphByName(UBlueprint const& FromBlueprint, const FName& GraphName)
 {
 	TArray<UEdGraph*> Graphs;
 	FromBlueprint.GetAllGraphs(Graphs);
 
 	UEdGraph* Ret = NULL;
-	if (UEdGraph** Result = Graphs.FindByPredicate(FMatchName(GraphName)))
+	if (UEdGraph** Result = Graphs.FindByPredicate(FMatchFName(GraphName)))
 	{
 		Ret = *Result;
 	}
@@ -84,26 +98,26 @@ static FDiffPanel InitializePanel(const FBlueprintRevPair& BlueprintRevPair)
 	return Ret;
 }
 
-static TArray< TSharedPtr<FMergeGraphRowEntry> > GenerateDiffListItems(const FBlueprintRevPair& RemoteBlueprint, const FBlueprintRevPair& BaseBlueprint, const FBlueprintRevPair& LocalBlueprint)
+static TArray< TSharedPtr<FMergeGraphRowEntry> > GenerateDiffListItems(const FBlueprintRevPair& RemoteBlueprint, const FBlueprintRevPair& BaseBlueprint, const FBlueprintRevPair& LocalBlueprint,  TArray< TSharedPtr<FGraphMergeConflict> >& OutMergeConflicts )
 {
 	// Index all the graphs by name, we use the name of the graph as the 
 	// basis of comparison between the various versions of the blueprint.
-	TMap< FString, UEdGraph* > RemoteGraphMap, BaseGraphMap, LocalGraphMap;
+	TMap< FName, UEdGraph* > RemoteGraphMap, BaseGraphMap, LocalGraphMap;
 	// We also want the set of all graph names in these blueprints, so that we 
 	// can iterate over every graph.
-	TSet< FString > AllGraphNames;
+	TSet< FName > AllGraphNames;
 	{
 		TArray<UEdGraph*> GraphsRemote, GraphsBase, GraphsLocal;
 		RemoteBlueprint.Blueprint->GetAllGraphs(GraphsRemote);
 		BaseBlueprint.Blueprint->GetAllGraphs(GraphsBase);
 		LocalBlueprint.Blueprint->GetAllGraphs(GraphsLocal);
 
-		const auto ToMap = [&AllGraphNames](const TArray<UEdGraph*>& InList, TMap<FString, UEdGraph*>& OutMap)
+		const auto ToMap = [&AllGraphNames](const TArray<UEdGraph*>& InList, TMap<FName, UEdGraph*>& OutMap)
 		{
 			for (auto Graph : InList)
 			{
-				OutMap.Add(Graph->GetName(), Graph);
-				AllGraphNames.Add(Graph->GetName());
+				OutMap.Add(Graph->GetFName(), Graph);
+				AllGraphNames.Add(Graph->GetFName());
 			}
 		};
 		ToMap(GraphsRemote, RemoteGraphMap);
@@ -146,19 +160,23 @@ static TArray< TSharedPtr<FMergeGraphRowEntry> > GenerateDiffListItems(const FBl
 				{
 					for (const auto& RemoteDifference : RemoteDifferences)
 					{
-						bool bConflicted = false;
-						if (RemoteDifference.Diff)
+						const FDiffSingleResult* ConflictingDifference = nullptr;
+
 						for (const auto& LocalDifference : LocalDifferences)
 						{
 							if (RemoteDifference.Node1 == LocalDifference.Node1 )
 							{
-								bConflicted |=	RemoteDifference.Diff == EDiffType::NODE_REMOVED 
-												|| LocalDifference.Diff == EDiffType::NODE_REMOVED;
-								bConflicted |= RemoteDifference.Pin1 == LocalDifference.Pin1;
+								if( RemoteDifference.Diff == EDiffType::NODE_REMOVED || 
+									LocalDifference.Diff == EDiffType::NODE_REMOVED ||
+									RemoteDifference.Pin1 == LocalDifference.Pin1 )
+								{
+									ConflictingDifference = &LocalDifference;
+									break;
+								}
 							}
 						}
 
-						if (bConflicted)
+						if (ConflictingDifference != nullptr)
 						{
 							// For now, we don't want to create a hard conflict for changes that don't effect runtime behavior:
 							if (RemoteDifference.Diff == EDiffType::NODE_MOVED ||
@@ -166,7 +184,8 @@ static TArray< TSharedPtr<FMergeGraphRowEntry> > GenerateDiffListItems(const FBl
 							{
 								continue;
 							}
-							// @todo doc: flag the conflict somehow...
+
+							OutMergeConflicts.Push(MakeShareable( new FGraphMergeConflict( GraphName, RemoteDifference, *ConflictingDifference)) );
 						}
 						else
 						{
@@ -238,7 +257,7 @@ static TArray< TSharedPtr<FMergeGraphRowEntry> > GenerateDiffListItems(const FBl
 
 			FMergeGraphRowEntry NewEntry = {
 				GraphName,
-				HackToShared(AllDifferences),
+				ConcreteToShared(AllDifferences),
 				bExistsInRemote,
 				bExistsInBase,
 				bExistsInLocal,
@@ -304,8 +323,10 @@ const FLinearColor MergeRed = FLinearColor(1.0f, 0.2f, 0.3f);
 const FLinearColor MergeYellow = FLinearColor::Yellow;
 const FLinearColor MergeWhite = FLinearColor::White;
 
-void SMergeGraphView::Construct(const FArguments InArgs, const FBlueprintMergeData& InData)
+void SMergeGraphView::Construct(const FArguments InArgs, const FBlueprintMergeData& InData )
 {
+	CurrentMergeConflict = INDEX_NONE;
+
 	Data = InData;
 	bViewsAreLocked = true;
 	DiffResultsListData = NULL;
@@ -331,7 +352,8 @@ void SMergeGraphView::Construct(const FArguments InArgs, const FBlueprintMergeDa
 
 	DifferencesFromBase = GenerateDiffListItems(BlueprintsForDisplay[EMergeParticipant::MERGE_PARTICIPANT_REMOTE]
 												, BlueprintsForDisplay[EMergeParticipant::MERGE_PARTICIPANT_BASE]
-												, BlueprintsForDisplay[EMergeParticipant::MERGE_PARTICIPANT_LOCAL]);
+												, BlueprintsForDisplay[EMergeParticipant::MERGE_PARTICIPANT_LOCAL]
+												, MergeConflicts);
 
 	// This is the function we'll use to generate a row in the control that lists all the available graphs:
 	const auto RowGenerator = [](TSharedPtr<FMergeGraphRowEntry> ParamItem, const TSharedRef<STableViewBase>& OwnerTable) -> TSharedRef<ITableRow>
@@ -388,7 +410,7 @@ void SMergeGraphView::Construct(const FArguments InArgs, const FBlueprintMergeDa
 				[
 					SNew(STextBlock)
 					.ColorAndOpacity(FLinearColor::White)
-					.Text(ParamItem->GraphName)
+					.Text(ParamItem->GraphName.GetPlainNameString())
 				]
 				+ Box(ParamItem->bExistsInRemote, RemoteColor)
 				+ Box(ParamItem->bExistsInBase, BaseColor)
@@ -491,6 +513,71 @@ bool SMergeGraphView::HasPrevDifference() const
 	return DiffWidgetUtils::HasPrevDifference(*DiffResultList.Pin(), *DiffResultsListData);
 }
 
+void SMergeGraphView::HighlightNextConflict()
+{
+	if( CurrentMergeConflict + 1 < MergeConflicts.Num())
+	{
+		++CurrentMergeConflict;
+	}
+	HighlightConflict(*MergeConflicts[CurrentMergeConflict]);
+}
+
+void SMergeGraphView::HighlightPrevConflict()
+{
+	if (CurrentMergeConflict - 1 >= 0 )
+	{
+		--CurrentMergeConflict;
+	}
+	HighlightConflict(*MergeConflicts[CurrentMergeConflict]);
+}
+
+bool SMergeGraphView::HasNextConflict() const 
+{
+	// return true if we have one conflict so that users can reselect the conflict if they desire. If we 
+	// return false when we already have selected this one and only conflict then there will be no way
+	// to reselect it if the user wants to.
+	return MergeConflicts.Num() == 1 || CurrentMergeConflict < MergeConflicts.Num();
+}
+
+bool SMergeGraphView::HasPrevConflict() const
+{
+	// note in HasNextConflict applies here as well.
+	return MergeConflicts.Num() == 1 || CurrentMergeConflict > 0;
+}
+
+void SMergeGraphView::HighlightConflict(const FGraphMergeConflict& Conflict)
+{
+	for( const auto& MergeGraphEntry : DifferencesFromBase )
+	{
+		if( MergeGraphEntry->GraphName == Conflict.GraphName )
+		{
+			OnGraphListSelectionChanged(MergeGraphEntry, ESelectInfo::Direct);
+			break;
+		}
+	}
+
+	// highlight the change made to the remote graph:
+	if( UEdGraphPin* Pin1 = Conflict.FirstResult.Pin2 )
+	{
+		// then look for the diff panel and focus on the change:
+		GetDiffPanelForNode(*Pin1->GetOwningNode(), DiffPanels).FocusDiff(*Pin1);
+	}
+	else if( UEdGraphNode* Node1 = Conflict.FirstResult.Node2 )
+	{
+		GetDiffPanelForNode(*Node1, DiffPanels).FocusDiff(*Node1);
+	}
+	
+	// and the change made to the local graph:
+	if (UEdGraphPin* Pin1 = Conflict.SecondResult.Pin2)
+	{
+		GetDiffPanelForNode(*Pin1->GetOwningNode(), DiffPanels).FocusDiff(*Pin1);
+	}
+	else if (UEdGraphNode* Node1 = Conflict.SecondResult.Node2)
+	{
+		GetDiffPanelForNode(*Node1, DiffPanels).FocusDiff(*Node1);
+	}
+}
+
 bool SMergeGraphView::HasNoDifferences() const
 {
 	auto DiffResultListPtr = DiffResultList.Pin();
@@ -499,12 +586,14 @@ bool SMergeGraphView::HasNoDifferences() const
 
 void SMergeGraphView::OnGraphListSelectionChanged(TSharedPtr<FMergeGraphRowEntry> Item, ESelectInfo::Type SelectionType)
 {
-	if (SelectionType != ESelectInfo::OnMouseClick || !Item.IsValid())
+	if (( SelectionType != ESelectInfo::OnMouseClick &&
+		  SelectionType != ESelectInfo::Direct) ||
+		!Item.IsValid())
 	{
 		return;
 	}
 
-	FString GraphName = Item->GraphName;
+	FName GraphName = Item->GraphName;
 
 	UEdGraph* GraphRemote = FindGraphByName(*GetRemotePanel().Blueprint, GraphName);
 	UEdGraph* GraphBase = FindGraphByName(*GetBasePanel().Blueprint, GraphName);
diff --git a/Engine/Source/Developer/Merge/Private/SMergeGraphView.h b/Engine/Source/Developer/Merge/Private/SMergeGraphView.h
--- a/Engine/Source/Developer/Merge/Private/SMergeGraphView.h
+++ b/Engine/Source/Developer/Merge/Private/SMergeGraphView.h
@@ -7,14 +7,14 @@
 
 class SMergeGraphView : public SCompoundWidget
 						, public IDiffControl
+						, public IMergeControl
 {
 public:
 	SLATE_BEGIN_ARGS(SMergeGraphView)
 	{}
 	SLATE_END_ARGS()
 
-	void Construct(const FArguments InArgs, const FBlueprintMergeData& InData);
-
+	void Construct(const FArguments InArgs, const FBlueprintMergeData& InData );
 private:
 	/** Implementation of IDiffControl: */
 	void NextDiff() override;
@@ -22,7 +22,14 @@ private:
 	bool HasNextDifference() const override;
 	bool HasPrevDifference() const override;
 
+	/** Implementation of IMergeControl */
+	void HighlightNextConflict() override;
+	void HighlightPrevConflict() override;
+	bool HasNextConflict() const override;
+	bool HasPrevConflict() const override;
+
 	/** Helper functions and event handlers: */
+	void HighlightConflict(const struct FGraphMergeConflict& Conflict);
 	bool HasNoDifferences() const;
 	void OnGraphListSelectionChanged(TSharedPtr<struct FMergeGraphRowEntry> Item, ESelectInfo::Type SelectionType);
 	void OnDiffListSelectionChanged(TSharedPtr<struct FDiffSingleResult> Item, ESelectInfo::Type SelectionType);
@@ -42,6 +49,9 @@ private:
 
 	TWeakPtr < SListView< TSharedPtr< struct FDiffSingleResult> > > DiffResultList;
 
+	TArray< TSharedPtr<struct FGraphMergeConflict> > MergeConflicts;
+	int CurrentMergeConflict;
+
 	/** 
 	 * Unfortunately this innocuous member requires some explanation:
 	 * 	This list is owned by the FMergeGraphRowEntry and displayed by DiffResultList, but we have
diff --git a/Engine/Source/Editor/GraphEditor/Public/DiffResults.h b/Engine/Source/Editor/GraphEditor/Public/DiffResults.h
--- a/Engine/Source/Editor/GraphEditor/Public/DiffResults.h
+++ b/Engine/Source/Editor/GraphEditor/Public/DiffResults.h
@@ -70,6 +70,18 @@ struct FDiffSingleResult
 	FLinearColor DisplayColor;
 };
 
+FORCEINLINE bool operator==( const FDiffSingleResult& LHS, const FDiffSingleResult& RHS )
+{
+	return	LHS.Diff == RHS.Diff &&
+			LHS.Node1 == RHS.Node1 &&
+			LHS.Node2 == RHS.Node2 &&
+			LHS.Pin1 == RHS.Pin1 &&
+			LHS.Pin2 == RHS.Pin2 &&
+			LHS.DisplayString == RHS.DisplayString &&
+			LHS.ToolTip == RHS.ToolTip &&
+			LHS.DisplayColor == RHS.DisplayColor;
+}
+
 /** Collects the Diffs found for a node */
 struct FDiffResults
 {
diff --git a/Engine/Source/Editor/Kismet/Private/SBlueprintDiff.cpp b/Engine/Source/Editor/Kismet/Private/SBlueprintDiff.cpp
--- a/Engine/Source/Editor/Kismet/Private/SBlueprintDiff.cpp
+++ b/Engine/Source/Editor/Kismet/Private/SBlueprintDiff.cpp
@@ -877,7 +877,7 @@ bool SBlueprintDiff::HasPrevDiff() const
 
 void SBlueprintDiff::FocusOnGraphRevisions( class UEdGraph* GraphOld, class UEdGraph* GraphNew, FListItemGraphToDiff* Diff )
 {
-	HandleGraphChanged( GraphOld ? GraphOld->GetName() : GraphNew->GetName() );
+	HandleGraphChanged( GraphOld ? GraphOld->GetFName() : GraphNew->GetFName() );
 
 	ResetGraphEditors();
 
@@ -1109,19 +1109,19 @@ FDiffPanel& SBlueprintDiff::GetDiffPanelForNode(UEdGraphNode& Node)
 	return Default;
 }
 
-void SBlueprintDiff::HandleGraphChanged( const FString& GraphName )
+void SBlueprintDiff::HandleGraphChanged( const FName GraphName )
 {
 	TArray<UEdGraph*> GraphsOld, GraphsNew;
 	PanelOld.Blueprint->GetAllGraphs(GraphsOld);
 	PanelNew.Blueprint->GetAllGraphs(GraphsNew);
 
 	UEdGraph* GraphOld = NULL;
-	if( UEdGraph** Iter = GraphsOld.FindByPredicate(FMatchName(GraphName)) )
+	if( UEdGraph** Iter = GraphsOld.FindByPredicate(FMatchFName(GraphName)) )
 	{
 		GraphOld = *Iter;
 	}
 	UEdGraph* GraphNew = NULL;
-	if (UEdGraph** Iter = GraphsNew.FindByPredicate(FMatchName(GraphName)))
+	if (UEdGraph** Iter = GraphsNew.FindByPredicate(FMatchFName(GraphName)))
 	{
 		GraphNew = *Iter;
 	}
diff --git a/Engine/Source/Editor/Kismet/Public/SBlueprintDiff.h b/Engine/Source/Editor/Kismet/Public/SBlueprintDiff.h
--- a/Engine/Source/Editor/Kismet/Public/SBlueprintDiff.h
+++ b/Engine/Source/Editor/Kismet/Public/SBlueprintDiff.h
@@ -5,19 +5,19 @@
 #include "Editor/GraphEditor/Public/DiffResults.h"
 #include "GraphEditor.h"
 
-struct FMatchName
+struct FMatchFName
 {
-	FMatchName(const FString& InName)
+	FMatchFName(FName InName)
 	: Name(InName)
 	{
 	}
 
 	bool operator() (const UObject* Object)
 	{
-		return Object->GetName() == Name;
+		return Object->GetFName() == Name;
 	}
 
-	FString const& Name;
+	FName const Name;
 };
 
 /** Individual Diff item shown in the list of diffs */
@@ -158,7 +158,7 @@ protected:
 	FDiffPanel& GetDiffPanelForNode(UEdGraphNode& Node);
 
 	/** Event handler that updates the graph view when user selects a new graph */
-	void HandleGraphChanged( const FString& GraphName );
+	void HandleGraphChanged( const FName GraphName );
 
 	TSharedRef<SWidget> GenerateGraphPanel();
 	TSharedRef<SWidget> GenerateDefaultsPanel();
