Separated list for remote differences and local differences in the graph view for the merge tool

diff --git a/Engine/Source/Developer/Merge/Private/SMergeGraphView.cpp b/Engine/Source/Developer/Merge/Private/SMergeGraphView.cpp
--- a/Engine/Source/Developer/Merge/Private/SMergeGraphView.cpp
+++ b/Engine/Source/Developer/Merge/Private/SMergeGraphView.cpp
@@ -9,16 +9,18 @@
 
 struct FGraphMergeConflict
 {
-	FGraphMergeConflict(FName InGraphName, const FDiffSingleResult& InFirstResult, const FDiffSingleResult& InSecondResult )
+	FGraphMergeConflict(FName InGraphName, const FDiffSingleResult& InRemoteDifference, const FDiffSingleResult& InLocalDifference )
 		: GraphName(InGraphName)
-		, FirstResult(InFirstResult)
-		, SecondResult(InSecondResult)
+		, RemoteDifference(InRemoteDifference)
+		, LocalDifference(InLocalDifference)
 	{
 	}
 
 	FName GraphName;
-	FDiffSingleResult FirstResult;
-	FDiffSingleResult SecondResult;
+
+	// These are the diff results that we have judged to be conflicting:
+	FDiffSingleResult RemoteDifference;
+	FDiffSingleResult LocalDifference;
 };
 
 // This cludge is used to move arrays of structures into arrays of 
@@ -53,7 +55,8 @@ struct FMergeGraphRowEntry
 
 	// These lists contain shared ptrs because they are displayed by a
 	// SListView, and it currently does not support lists of values.
-	TArray< TSharedPtr<FDiffSingleResult> > Differences;
+	TArray< TSharedPtr<FDiffSingleResult> > RemoteDifferences;
+	TArray< TSharedPtr<FDiffSingleResult> > LocalDifferences;
 
 	bool bExistsInRemote;
 	bool bExistsInBase;
@@ -67,7 +70,7 @@ static UEdGraph* FindGraphByName(UBlueprint const& FromBlueprint, const FName& G
 	TArray<UEdGraph*> Graphs;
 	FromBlueprint.GetAllGraphs(Graphs);
 
-	UEdGraph* Ret = NULL;
+	UEdGraph* Ret = nullptr;
 	if (UEdGraph** Result = Graphs.FindByPredicate(FMatchFName(GraphName)))
 	{
 		Ret = *Result;
@@ -130,7 +133,7 @@ static TArray< TSharedPtr<FMergeGraphRowEntry> > GenerateDiffListItems(const FBl
 		const auto GenerateDifferences = [](UEdGraph* GraphNew, UEdGraph** GraphOld)
 		{
 			TArray<FDiffSingleResult> Results;
-			FGraphDiffControl::DiffGraphs(GraphOld ? *GraphOld : NULL, GraphNew, Results);
+			FGraphDiffControl::DiffGraphs(GraphOld ? *GraphOld : nullptr, GraphNew, Results);
 			return Results;
 		};
 
@@ -174,6 +177,14 @@ static TArray< TSharedPtr<FMergeGraphRowEntry> > GenerateDiffListItems(const FBl
 									break;
 								}
 							}
+							else if( RemoteDifference.Pin1 != nullptr && (RemoteDifference.Pin1 == LocalDifference.Pin1 ) )
+							{
+								// it's possible the users made the same change to the same pin, but given the wide
+								// variety of changes that can be made to a pin it is difficult to identify the change 
+								// as identical, for now I'm just flagging all changes to the same pin as a conflict:
+								ConflictingDifference = &LocalDifference;
+								break;
+							}
 						}
 
 						if (ConflictingDifference != nullptr)
@@ -247,17 +258,15 @@ static TArray< TSharedPtr<FMergeGraphRowEntry> > GenerateDiffListItems(const FBl
 					}
 				}
 
-				bExistsInRemote = RemoteGraph != NULL;
-				bExistsInBase = BaseGraph != NULL;
-				bExistsInLocal = LocalGraph != NULL;
+				bExistsInRemote = RemoteGraph != nullptr;
+				bExistsInBase = BaseGraph != nullptr;
+				bExistsInLocal = LocalGraph != nullptr;
 			}
 
-			TArray<FDiffSingleResult> AllDifferences(RemoteDifferences);
-			AllDifferences.Append(LocalDifferences);
-
 			FMergeGraphRowEntry NewEntry = {
 				GraphName,
-				ConcreteToShared(AllDifferences),
+				ConcreteToShared(RemoteDifferences),
+				ConcreteToShared(LocalDifferences),
 				bExistsInRemote,
 				bExistsInBase,
 				bExistsInLocal,
@@ -329,7 +338,8 @@ void SMergeGraphView::Construct(const FArguments InArgs, const FBlueprintMergeDa
 
 	Data = InData;
 	bViewsAreLocked = true;
-	DiffResultsListData = NULL;
+	RemoteDiffResultsListData = nullptr;
+	LocalDiffResultsListData = nullptr;
 
 	TArray<FBlueprintRevPair> BlueprintsForDisplay;
 	// EMergeParticipant::MERGE_PARTICIPANT_REMOTE
@@ -418,6 +428,30 @@ void SMergeGraphView::Construct(const FArguments InArgs, const FBlueprintMergeDa
 			];
 	};
 
+	const auto DifferencesListGenerator = []( FText Title, TSharedPtr<SBox>& OutResultsContainer )
+	{
+		return SNew(SVerticalBox)
+			+ SVerticalBox::Slot()
+			.AutoHeight()
+			.HAlign(HAlign_Left)
+			[
+				SNew(STextBlock)
+				.Text(Title)
+			]
+		+ SVerticalBox::Slot()
+			.AutoHeight()
+			[
+				// this creates a line between the header and the list itself:
+				SNew(SBorder)
+				.Padding(FEditorStyle::GetMargin(TEXT("Menu.Separator.Padding")))
+				.BorderImage(FEditorStyle::GetBrush(TEXT("Menu.Separator")))
+			]
+		+ SVerticalBox::Slot()
+			[
+				SAssignNew(OutResultsContainer, SBox)
+			];
+	};
+
 	ChildSlot
 		[
 			SNew(SSplitter)
@@ -464,11 +498,23 @@ void SMergeGraphView::Construct(const FArguments InArgs, const FBlueprintMergeDa
 						]
 					]
 					+ SSplitter::Slot()
-						.Value(.7f)
+					.Value(.35f)
+					[
+						SNew(SBorder)
+						.BorderImage(FEditorStyle::GetBrush("ToolPanel.GroupBorder"))
 						[
-							SAssignNew(DiffResultsWidget, SBorder)
-							.BorderImage(FEditorStyle::GetBrush("ToolPanel.GroupBorder"))
+							DifferencesListGenerator(LOCTEXT("RemoteChangesLabel", "Remote Changes"), RemoteDiffResultsWidget)
 						]
+					]
+					+ SSplitter::Slot()
+					.Value(.35f)
+					[
+						SNew(SBorder)
+						.BorderImage(FEditorStyle::GetBrush("ToolPanel.GroupBorder"))
+						[
+							DifferencesListGenerator(LOCTEXT("LocalChangesLabel", "Local Changes"), LocalDiffResultsWidget)
+						]
+					]
 				]
 			]
 			+ SSplitter::Slot()
@@ -481,14 +527,30 @@ void SMergeGraphView::Construct(const FArguments InArgs, const FBlueprintMergeDa
 
 void SMergeGraphView::NextDiff()
 {
-	check(DiffResultList.Pin().IsValid() && DiffResultsListData);
-	DiffWidgetUtils::SelectNextRow( *DiffResultList.Pin(), *DiffResultsListData );
+	check(	RemoteDiffResultList.Pin().IsValid() && RemoteDiffResultsListData || 
+			LocalDiffResultList.Pin().IsValid() && LocalDiffResultsListData );
+	if( RemoteDiffResultsListData && DiffWidgetUtils::HasNextDifference( *RemoteDiffResultList.Pin(), *RemoteDiffResultsListData) )
+	{
+		DiffWidgetUtils::SelectNextRow( *RemoteDiffResultList.Pin(), *RemoteDiffResultsListData );
+	}
+	else
+	{
+		DiffWidgetUtils::SelectNextRow(*LocalDiffResultList.Pin(), *LocalDiffResultsListData);
+	}
 }
 
 void SMergeGraphView::PrevDiff()
 {
-	check(DiffResultList.Pin().IsValid() && DiffResultsListData);
-	DiffWidgetUtils::SelectPrevRow( *DiffResultList.Pin(), *DiffResultsListData );
+	check(	RemoteDiffResultList.Pin().IsValid() && RemoteDiffResultsListData ||
+			LocalDiffResultList.Pin().IsValid() && LocalDiffResultsListData);
+	if (LocalDiffResultsListData && DiffWidgetUtils::HasPrevDifference(*LocalDiffResultList.Pin(), *LocalDiffResultsListData))
+	{
+		DiffWidgetUtils::SelectPrevRow(*LocalDiffResultList.Pin(), *LocalDiffResultsListData);
+	}
+	else
+	{
+		DiffWidgetUtils::SelectPrevRow(*RemoteDiffResultList.Pin(), *RemoteDiffResultsListData);
+	}
 }
 
 bool SMergeGraphView::HasNextDifference() const
@@ -498,8 +560,10 @@ bool SMergeGraphView::HasNextDifference() const
 		return false;
 	}
 
-	check(DiffResultList.Pin().IsValid() && DiffResultsListData);
-	return DiffWidgetUtils::HasNextDifference(*DiffResultList.Pin(), *DiffResultsListData);
+	check(	RemoteDiffResultList.Pin().IsValid() && RemoteDiffResultsListData ||
+			LocalDiffResultList.Pin().IsValid() && LocalDiffResultsListData);
+	return	(RemoteDiffResultsListData && DiffWidgetUtils::HasNextDifference(*RemoteDiffResultList.Pin(), *RemoteDiffResultsListData) ) ||
+			(LocalDiffResultsListData && DiffWidgetUtils::HasNextDifference(*LocalDiffResultList.Pin(), *LocalDiffResultsListData) ) ;
 }
 
 bool SMergeGraphView::HasPrevDifference() const
@@ -509,8 +573,10 @@ bool SMergeGraphView::HasPrevDifference() const
 		return false;
 	}
 
-	check(DiffResultList.Pin().IsValid() && DiffResultsListData);
-	return DiffWidgetUtils::HasPrevDifference(*DiffResultList.Pin(), *DiffResultsListData);
+	check(	RemoteDiffResultList.Pin().IsValid() && RemoteDiffResultsListData ||
+			LocalDiffResultList.Pin().IsValid() && LocalDiffResultsListData);
+	return DiffWidgetUtils::HasPrevDifference(*RemoteDiffResultList.Pin(), *RemoteDiffResultsListData) 
+		|| DiffWidgetUtils::HasPrevDifference(*LocalDiffResultList.Pin(), *LocalDiffResultsListData);
 }
 
 void SMergeGraphView::HighlightNextConflict()
@@ -557,31 +623,33 @@ void SMergeGraphView::HighlightConflict(const FGraphMergeConflict& Conflict)
 	}
 
 	// highlight the change made to the remote graph:
-	if( UEdGraphPin* Pin1 = Conflict.FirstResult.Pin2 )
+	if( UEdGraphPin* Pin = Conflict.RemoteDifference.Pin2 )
 	{
 		// then look for the diff panel and focus on the change:
-		GetDiffPanelForNode(*Pin1->GetOwningNode(), DiffPanels).FocusDiff(*Pin1);
+		GetDiffPanelForNode(*Pin->GetOwningNode(), DiffPanels).FocusDiff(*Pin);
 	}
-	else if( UEdGraphNode* Node1 = Conflict.FirstResult.Node2 )
+	else if( UEdGraphNode* Node = Conflict.RemoteDifference.Node2 )
 	{
-		GetDiffPanelForNode(*Node1, DiffPanels).FocusDiff(*Node1);
+		GetDiffPanelForNode(*Node, DiffPanels).FocusDiff(*Node);
 	}
 	
 	// and the change made to the local graph:
-	if (UEdGraphPin* Pin1 = Conflict.SecondResult.Pin2)
+	if (UEdGraphPin* Pin = Conflict.LocalDifference.Pin2)
 	{
-		GetDiffPanelForNode(*Pin1->GetOwningNode(), DiffPanels).FocusDiff(*Pin1);
+		GetDiffPanelForNode(*Pin->GetOwningNode(), DiffPanels).FocusDiff(*Pin);
 	}
-	else if (UEdGraphNode* Node1 = Conflict.SecondResult.Node2)
+	else if (UEdGraphNode* Node = Conflict.LocalDifference.Node2)
 	{
-		GetDiffPanelForNode(*Node1, DiffPanels).FocusDiff(*Node1);
+		GetDiffPanelForNode(*Node, DiffPanels).FocusDiff(*Node);
 	}
 }
 
 bool SMergeGraphView::HasNoDifferences() const
 {
-	auto DiffResultListPtr = DiffResultList.Pin();
-	return !DiffResultListPtr.IsValid() || DiffResultListPtr->GetNumItemsBeingObserved() == 0;
+	auto RemoteDiffResultListPtr = RemoteDiffResultList.Pin();
+	auto LocalDiffResultListPtr = LocalDiffResultList.Pin();
+	return	(!RemoteDiffResultListPtr.IsValid() || RemoteDiffResultListPtr->GetNumItemsBeingObserved() == 0) ||
+			(!LocalDiffResultListPtr.IsValid() || LocalDiffResultListPtr->GetNumItemsBeingObserved() == 0);
 }
 
 void SMergeGraphView::OnGraphListSelectionChanged(TSharedPtr<FMergeGraphRowEntry> Item, ESelectInfo::Type SelectionType)
@@ -599,30 +667,41 @@ void SMergeGraphView::OnGraphListSelectionChanged(TSharedPtr<FMergeGraphRowEntry
 	UEdGraph* GraphBase = FindGraphByName(*GetBasePanel().Blueprint, GraphName);
 	UEdGraph* GraphLocal = FindGraphByName(*GetLocalPanel().Blueprint, GraphName);
 
-	GetBasePanel().GeneratePanel(GraphBase, NULL);
+	GetBasePanel().GeneratePanel(GraphBase, nullptr);
 	GetRemotePanel().GeneratePanel(GraphRemote, GraphBase);
 	GetLocalPanel().GeneratePanel(GraphLocal, GraphBase);
 
 	LockViews(DiffPanels, bViewsAreLocked);
 
-	DiffResultsListData = &(Item->Differences);
-	DiffResultsWidget->SetContent(
-		SAssignNew(DiffResultList, SListView< TSharedPtr<FDiffSingleResult> >)
-		.ItemHeight(24)
-		.ListItemsSource(&Item->Differences)
-		.OnGenerateRow(SListView< TSharedPtr<FDiffSingleResult> >::FOnGenerateRow::CreateStatic(
+	RemoteDiffResultsListData = &(Item->RemoteDifferences);
+	LocalDiffResultsListData = &(Item->LocalDifferences);
+
+	const auto SetupDiffList = [this]( TArray< TSharedPtr<FDiffSingleResult> > *Differences, TWeakPtr < SListView< TSharedPtr< struct FDiffSingleResult> > >& OutWidget)
+	{
+		return SAssignNew(OutWidget, SListView< TSharedPtr<FDiffSingleResult> >)
+			.ItemHeight(24)
+			.ListItemsSource(Differences)
+			.OnGenerateRow(SListView< TSharedPtr<FDiffSingleResult> >::FOnGenerateRow::CreateStatic(
 			[](TSharedPtr<FDiffSingleResult> ParamItem, const TSharedRef<STableViewBase>& OwnerTable) -> TSharedRef<ITableRow>
-			{
-				FDiffResultItem WidgetGenerator(*ParamItem);
-				return SNew(STableRow< TSharedPtr<FDiffSingleResult> >, OwnerTable)
-					.Content()
-					[
-						WidgetGenerator.GenerateWidget()
-					];
-			}
+		{
+			FDiffResultItem WidgetGenerator(*ParamItem);
+			return SNew(STableRow< TSharedPtr<FDiffSingleResult> >, OwnerTable)
+				.Content()
+				[
+					WidgetGenerator.GenerateWidget()
+				];
+		}
 		))
-		.SelectionMode(ESelectionMode::Single)
-		.OnSelectionChanged(this, &SMergeGraphView::OnDiffListSelectionChanged)
+			.SelectionMode(ESelectionMode::Single)
+			.OnSelectionChanged(this, &SMergeGraphView::OnDiffListSelectionChanged);
+	};
+
+	RemoteDiffResultsWidget->SetContent(
+		SetupDiffList( &Item->RemoteDifferences, RemoteDiffResultList)
+	);
+
+	LocalDiffResultsWidget->SetContent(
+		SetupDiffList(&Item->LocalDifferences, LocalDiffResultList)
 	);
 }
 
diff --git a/Engine/Source/Developer/Merge/Private/SMergeGraphView.h b/Engine/Source/Developer/Merge/Private/SMergeGraphView.h
--- a/Engine/Source/Developer/Merge/Private/SMergeGraphView.h
+++ b/Engine/Source/Developer/Merge/Private/SMergeGraphView.h
@@ -45,9 +45,12 @@ private:
 	FBlueprintMergeData Data;
 
 	TArray< TSharedPtr< struct FMergeGraphRowEntry> > DifferencesFromBase;
-	TSharedPtr<class SBorder> DiffResultsWidget;
 
-	TWeakPtr < SListView< TSharedPtr< struct FDiffSingleResult> > > DiffResultList;
+	TSharedPtr<class SBox> RemoteDiffResultsWidget;
+	TSharedPtr<class SBox> LocalDiffResultsWidget;
+
+	TWeakPtr < SListView< TSharedPtr< struct FDiffSingleResult> > > RemoteDiffResultList;
+	TWeakPtr < SListView< TSharedPtr< struct FDiffSingleResult> > > LocalDiffResultList;
 
 	TArray< TSharedPtr<struct FGraphMergeConflict> > MergeConflicts;
 	int CurrentMergeConflict;
@@ -66,7 +69,8 @@ private:
 	 * 
 	 * 	I have chosen option 1.
 	 */
-	TArray< TSharedPtr< struct FDiffSingleResult> > const* DiffResultsListData;
+	TArray< TSharedPtr< struct FDiffSingleResult> > const* RemoteDiffResultsListData;
+	TArray< TSharedPtr< struct FDiffSingleResult> > const* LocalDiffResultsListData;
 
 	bool bViewsAreLocked;
 };
