Merge panel's details view can now highlight conflicts

diff --git a/Engine/Source/Developer/Merge/Private/SBlueprintMerge.cpp b/Engine/Source/Developer/Merge/Private/SBlueprintMerge.cpp
--- a/Engine/Source/Developer/Merge/Private/SBlueprintMerge.cpp
+++ b/Engine/Source/Developer/Merge/Private/SBlueprintMerge.cpp
@@ -61,6 +61,7 @@ void SBlueprintMerge::Construct(const FArguments InArgs, const FBlueprintMergeDa
 	auto DetailsView = SNew(SMergeDetailsView, InData);
 	DetailsControl.Widget = DetailsView;
 	DetailsControl.DiffControl = &DetailsView.Get();
+	DetailsControl.MergeControl = &DetailsView.Get();
 
 	FToolBarBuilder ToolbarBuilder(TSharedPtr< FUICommandList >(), FMultiBoxCustomization::None);
 	ToolbarBuilder.AddToolBarButton( 
diff --git a/Engine/Source/Developer/Merge/Private/SMergeDetailsView.cpp b/Engine/Source/Developer/Merge/Private/SMergeDetailsView.cpp
--- a/Engine/Source/Developer/Merge/Private/SMergeDetailsView.cpp
+++ b/Engine/Source/Developer/Merge/Private/SMergeDetailsView.cpp
@@ -10,6 +10,7 @@
 void SMergeDetailsView::Construct(const FArguments InArgs, const FBlueprintMergeData& InData )
 {
 	Data = InData;
+	CurrentMergeConflict = INDEX_NONE;
 
 	const UObject* RemoteCDO = DiffUtils::GetCDO(InData.BlueprintRemote);
 	const UObject* BaseCDO = DiffUtils::GetCDO(InData.BlueprintBase);
@@ -23,7 +24,8 @@ void SMergeDetailsView::Construct(const FArguments InArgs, const FBlueprintMerge
 	FPropertySoftPathSet RemoteDifferingPropertiesSet(RemoteDifferingProperties);
 	FPropertySoftPathSet LocalDifferingPropertiesSet(LocalDifferingProperties);
 	FPropertySoftPathSet BaseDifferingPropertiesSet(RemoteDifferingPropertiesSet.Union(LocalDifferingPropertiesSet));
-	
+	MergeConflicts = RemoteDifferingPropertiesSet.Intersect(LocalDifferingPropertiesSet).Array();
+
 	/*	
 		DifferingProperties is an ordered list of all differing properties (properties added, removed or changed) in remote or local.
 		Strictly speaking it's impossible to guarantee that we'll traverse remote and local differences in the same order (for instance
@@ -153,6 +155,44 @@ bool SMergeDetailsView::HasPrevDifference() const
 	return DifferingProperties.IsValidIndex(CurrentDifference - 1);
 }
 
+void SMergeDetailsView::HighlightNextConflict()
+{
+	if (CurrentMergeConflict + 1 < MergeConflicts.Num())
+	{
+		++CurrentMergeConflict;
+	}
+	for (auto& DetailDiff : DetailsViews)
+	{
+		DetailDiff.HighlightProperty(MergeConflicts[CurrentMergeConflict]);
+	}
+}
+
+void SMergeDetailsView::HighlightPrevConflict()
+{
+	if (CurrentMergeConflict - 1 >= 0)
+	{
+		--CurrentMergeConflict;
+	}
+	for (auto& DetailDiff : DetailsViews)
+	{
+		DetailDiff.HighlightProperty(MergeConflicts[CurrentMergeConflict]);
+	}
+}
+
+bool SMergeDetailsView::HasNextConflict() const
+{
+	// return true if we have one conflict so that users can reselect the conflict if they desire. If we 
+	// return false when we already have selected this one and only conflict then there will be no way
+	// to reselect it if the user wants to.
+	return MergeConflicts.Num() != 0 && (MergeConflicts.Num() == 1 || CurrentMergeConflict < MergeConflicts.Num());
+}
+
+bool SMergeDetailsView::HasPrevConflict() const
+{
+	// note in HasNextConflict applies here as well.
+	return MergeConflicts.Num() == 1 || CurrentMergeConflict > 0;
+}
+
 FDetailsDiff& SMergeDetailsView::GetRemoteDetails()
 {
 	return DetailsViews[EMergeParticipant::MERGE_PARTICIPANT_REMOTE];
diff --git a/Engine/Source/Developer/Merge/Private/SMergeDetailsView.h b/Engine/Source/Developer/Merge/Private/SMergeDetailsView.h
--- a/Engine/Source/Developer/Merge/Private/SMergeDetailsView.h
+++ b/Engine/Source/Developer/Merge/Private/SMergeDetailsView.h
@@ -7,6 +7,7 @@
 
 class SMergeDetailsView : public SCompoundWidget
 						, public IDiffControl
+						, public IMergeControl
 {
 public:
 	virtual ~SMergeDetailsView() {}
@@ -23,11 +24,20 @@ private:
 	bool HasNextDifference() const override;
 	bool HasPrevDifference() const override;
 
+	/** Implementation of IMergeControl: */
+	void HighlightNextConflict() override;
+	void HighlightPrevConflict() override;
+	bool HasNextConflict() const override;
+	bool HasPrevConflict() const override;
+
 	void HighlightCurrentDifference();
 	FDetailsDiff& GetRemoteDetails();
 	FDetailsDiff& GetBaseDetails();
 	FDetailsDiff& GetLocalDetails();
 
+	TArray< FPropertySoftPath > MergeConflicts;
+	int CurrentMergeConflict;
+
 	FBlueprintMergeData Data;
 	TArray< FDetailsDiff > DetailsViews;
 
diff --git a/Engine/Source/Developer/Merge/Private/SMergeGraphView.cpp b/Engine/Source/Developer/Merge/Private/SMergeGraphView.cpp
--- a/Engine/Source/Developer/Merge/Private/SMergeGraphView.cpp
+++ b/Engine/Source/Developer/Merge/Private/SMergeGraphView.cpp
@@ -602,7 +602,7 @@ bool SMergeGraphView::HasNextConflict() const
 	// return true if we have one conflict so that users can reselect the conflict if they desire. If we 
 	// return false when we already have selected this one and only conflict then there will be no way
 	// to reselect it if the user wants to.
-	return MergeConflicts.Num() == 1 || CurrentMergeConflict < MergeConflicts.Num();
+	return  MergeConflicts.Num() != 0 && (MergeConflicts.Num() == 1 || CurrentMergeConflict < MergeConflicts.Num());
 }
 
 bool SMergeGraphView::HasPrevConflict() const
