Conflict detection and diff highlighting fixed for the tree view. Also fixed Native Components not being checked. And an off by one error in HasNextDifference

diff --git a/Engine/Source/Developer/Merge/Private/SBlueprintMerge.cpp b/Engine/Source/Developer/Merge/Private/SBlueprintMerge.cpp
--- a/Engine/Source/Developer/Merge/Private/SBlueprintMerge.cpp
+++ b/Engine/Source/Developer/Merge/Private/SBlueprintMerge.cpp
@@ -57,6 +57,7 @@ void SBlueprintMerge::Construct(const FArguments InArgs, const FBlueprintMergeDa
 	auto TreeView = SNew(SMergeTreeView, InData);
 	TreeControl.Widget = TreeView;
 	TreeControl.DiffControl = &TreeView.Get();
+	TreeControl.MergeControl = &TreeView.Get();
 
 	auto DetailsView = SNew(SMergeDetailsView, InData);
 	DetailsControl.Widget = DetailsView;
diff --git a/Engine/Source/Developer/Merge/Private/SMergeDetailsView.cpp b/Engine/Source/Developer/Merge/Private/SMergeDetailsView.cpp
--- a/Engine/Source/Developer/Merge/Private/SMergeDetailsView.cpp
+++ b/Engine/Source/Developer/Merge/Private/SMergeDetailsView.cpp
@@ -184,7 +184,7 @@ bool SMergeDetailsView::HasNextConflict() const
 	// return true if we have one conflict so that users can reselect the conflict if they desire. If we 
 	// return false when we already have selected this one and only conflict then there will be no way
 	// to reselect it if the user wants to.
-	return MergeConflicts.Num() != 0 && (MergeConflicts.Num() == 1 || CurrentMergeConflict < MergeConflicts.Num());
+	return MergeConflicts.Num() != 0 && (MergeConflicts.Num() == 1 || CurrentMergeConflict + 1 < MergeConflicts.Num());
 }
 
 bool SMergeDetailsView::HasPrevConflict() const
diff --git a/Engine/Source/Developer/Merge/Private/SMergeGraphView.cpp b/Engine/Source/Developer/Merge/Private/SMergeGraphView.cpp
--- a/Engine/Source/Developer/Merge/Private/SMergeGraphView.cpp
+++ b/Engine/Source/Developer/Merge/Private/SMergeGraphView.cpp
@@ -198,63 +198,8 @@ static TArray< TSharedPtr<FMergeGraphRowEntry> > GenerateDiffListItems(const FBl
 
 							OutMergeConflicts.Push(MakeShareable( new FGraphMergeConflict( GraphName, RemoteDifference, *ConflictingDifference)) );
 						}
-						else
-						{
-							// no conflict, we want to be able to automatically apply this remote change
-							// to the local revision:
-							switch (RemoteDifference.Diff)
-							{
-							case EDiffType::NODE_REMOVED:
-								{
-									// Find the corresponding object in the local graph:
-							
-									const UEdGraphNode* TargetNode = FindNodeInUnrelatedGraph(*LocalGraph, RemoteDifference.Node1);
-									check(TargetNode); // should have been a conflict!
-								}
-								break;
-							case EDiffType::NODE_ADDED:
-
-
-								break;
-							case EDiffType::PIN_LINKEDTO_NUM_DEC:
-								
-								break;
-							case EDiffType::PIN_LINKEDTO_NUM_INC:
-								break;
-							case EDiffType::PIN_DEFAULT_VALUE:
-								break;
-							case EDiffType::PIN_TYPE_CATEGORY:
-								break;
-							case EDiffType::PIN_TYPE_SUBCATEGORY:
-								break;
-							case EDiffType::PIN_TYPE_SUBCATEGORY_OBJECT:
-								break;
-							case EDiffType::PIN_TYPE_IS_ARRAY:
-								break;
-							case EDiffType::PIN_TYPE_IS_REF:
-								break;
-							case EDiffType::PIN_LINKEDTO_NODE:
-								break;
-							case EDiffType::NODE_MOVED:
-
-								break;
-							case EDiffType::TIMELINE_LENGTH:
-							case EDiffType::TIMELINE_AUTOPLAY:
-							case EDiffType::TIMELINE_LOOP:
-							case EDiffType::TIMELINE_NUM_TRACKS:
-							case EDiffType::TIMELINE_TRACK_MODIFIED:
-								check(false); // doesn't apply to blueprint
-								break;
-							case EDiffType::NODE_PIN_COUNT:
-								break;
-							case EDiffType::NODE_COMMENT:
-								break;
-							case EDiffType::NODE_PROPERTY:
-								break;
-							case EDiffType::NO_DIFFERENCE:
-								check(false);
-							}
-						}
+						// else: no conflict, we want to be able to automatically apply this remote change
+						// to the local revision:
 					}
 				}
 
@@ -602,7 +547,7 @@ bool SMergeGraphView::HasNextConflict() const
 	// return true if we have one conflict so that users can reselect the conflict if they desire. If we 
 	// return false when we already have selected this one and only conflict then there will be no way
 	// to reselect it if the user wants to.
-	return  MergeConflicts.Num() != 0 && (MergeConflicts.Num() == 1 || CurrentMergeConflict < MergeConflicts.Num());
+	return  MergeConflicts.Num() != 0 && (MergeConflicts.Num() == 1 || CurrentMergeConflict + 1 < MergeConflicts.Num());
 }
 
 bool SMergeGraphView::HasPrevConflict() const
diff --git a/Engine/Source/Developer/Merge/Private/SMergeTreeView.cpp b/Engine/Source/Developer/Merge/Private/SMergeTreeView.cpp
--- a/Engine/Source/Developer/Merge/Private/SMergeTreeView.cpp
+++ b/Engine/Source/Developer/Merge/Private/SMergeTreeView.cpp
@@ -8,53 +8,7 @@ void SMergeTreeView::Construct(const FArguments InArgs, const FBlueprintMergeDat
 {
 	Data = InData;
 	CurrentDifference = -1;
-
-	TArray< FSCSDiffEntry > RemoteDifferingProperties;
-	TArray< int > RemoteDiffsSortKeys;
-	TArray< FSCSDiffEntry > LocalDifferingProperties;
-	TArray< int > LocalDiffsSortKeys;
-
-	DiffUtils::CompareUnrelatedSCS(InData.BlueprintBase, InData.BlueprintRemote, RemoteDifferingProperties, &RemoteDiffsSortKeys);
-	DiffUtils::CompareUnrelatedSCS(InData.BlueprintBase, InData.BlueprintLocal, LocalDifferingProperties, &LocalDiffsSortKeys);
-
-	check(RemoteDiffsSortKeys.Num() == RemoteDifferingProperties.Num() && LocalDifferingProperties.Num() == LocalDiffsSortKeys.Num());
-
-	// Order the differing properties, and eliminate dupes. Eventually dupes should be treated as collisions when
-	// attempting to complete a merge automatically:
-	int RemoteIter = 0;
-	int LocalIter = 0;
-	while (RemoteDifferingProperties.IsValidIndex(RemoteIter) || LocalDifferingProperties.IsValidIndex(LocalIter))
-	{
-		const bool bStillHaveLocalDiffs = LocalDifferingProperties.IsValidIndex(LocalIter);
-		const bool bStillHaveRemoteDiffs = RemoteDifferingProperties.IsValidIndex(RemoteIter);
-		FSCSDiffEntry Entry;
-
-		if (bStillHaveLocalDiffs && bStillHaveRemoteDiffs)
-		{
-			if (LocalDiffsSortKeys[LocalIter] <= RemoteDiffsSortKeys[RemoteIter])
-			{
-				Entry = LocalDifferingProperties[LocalIter++];
-			}
-			else
-			{
-				Entry = RemoteDifferingProperties[RemoteIter++];
-			}
-		}
-		else if (bStillHaveLocalDiffs)
-		{
-			Entry = LocalDifferingProperties[LocalIter++];
-		}
-		else
-		{
-			check(bStillHaveRemoteDiffs);
-			Entry = RemoteDifferingProperties[RemoteIter++];
-		}
-
-		if (!DifferingProperties.Contains(Entry))
-		{
-			DifferingProperties.Push(Entry);
-		}
-	}
+	CurrentMergeConflict = -1;
 
 	// generate controls:
 	// EMergeParticipant::MERGE_PARTICIPANT_REMOTE
@@ -76,6 +30,98 @@ void SMergeTreeView::Construct(const FArguments InArgs, const FBlueprintMergeDat
 			);
 	}
 
+	TArray< FSCSResolvedIdentifier > RemoteHierarchy = GetRemoteView().GetDisplayedHierarchy();
+	TArray< FSCSResolvedIdentifier > BaseHierarchy = GetBaseView().GetDisplayedHierarchy();
+	TArray< FSCSResolvedIdentifier > LocalHierarchy = GetLocalView().GetDisplayedHierarchy();
+
+	FSCSDiffRoot RemoteDifferingProperties;
+	DiffUtils::CompareUnrelatedSCS(InData.BlueprintBase, BaseHierarchy, InData.BlueprintRemote, RemoteHierarchy, RemoteDifferingProperties );
+	FSCSDiffRoot LocalDifferingProperties;
+	DiffUtils::CompareUnrelatedSCS(InData.BlueprintBase, BaseHierarchy, InData.BlueprintLocal, LocalHierarchy, LocalDifferingProperties);
+
+	DifferingProperties = RemoteDifferingProperties;
+	DifferingProperties.Entries.Append( LocalDifferingProperties.Entries );
+
+	// Remove duplicates and detect conflicts, there is opportunity for performance improvement here, but I doubt this represents a 
+	// significant cost:
+	for( int32 Iter = 0; Iter != DifferingProperties.Entries.Num(); ++Iter )
+	{
+		const FSCSDiffEntry& IterEntry = DifferingProperties.Entries[Iter];
+		for( int32 Jter = Iter + 1; Jter < DifferingProperties.Entries.Num(); ++Jter )
+		{
+			const FSCSDiffEntry& JterEntry = DifferingProperties.Entries[Jter];
+			// If the change is to a specific property, then it's only a conflict if the other change
+			// is also to that same specific property. All 'whole object' (non property specific changes)
+			// are conflicting:
+			if( IterEntry.TreeIdentifier == JterEntry.TreeIdentifier &&
+				( ( IterEntry.PropertyIdentifier == JterEntry.PropertyIdentifier &&
+				  IterEntry.PropertyIdentifier != FPropertySoftPath() ) 
+				||
+				( ( IterEntry.PropertyIdentifier == FPropertySoftPath() ||
+					JterEntry.PropertyIdentifier == FPropertySoftPath() ) ) ))
+			{
+				// for now just flag it as a conflict:
+				MergeConflicts.Entries.Push( IterEntry );
+
+				// and remove the redundant difference:
+				DifferingProperties.Entries.RemoveAtSwap(Jter--);
+			}
+		}
+	}
+
+	/* This predicate sorts the list of differing properties so that those that are 'earlier' in the tree appear first.
+		For example, if we get the following two trees back:
+			
+			B added at position (3, 2, 1)
+			C removed at position (1, 2)
+
+			and 
+
+			D added at position (4, 2, 1)
+
+			the resulting list will be
+			C, B, aand D last: */
+	const auto SortTreePredicate = []( const FSCSDiffEntry& A, const FSCSDiffEntry& B )
+	{
+		int32 Idx = 0;
+		const TArray<int32>& ATreeAddress = A.TreeIdentifier.TreeLocation;
+		const TArray<int32>& BTreeAddress = B.TreeIdentifier.TreeLocation;
+		while(true)
+		{
+			if( !ATreeAddress.IsValidIndex(Idx) )
+			{
+				// A has a shorter address, show it first:
+				return true;
+			}
+			else if( !BTreeAddress.IsValidIndex(Idx) )
+			{
+				// B has a shorter address, show it first:
+				return false;
+			}
+			else if( ATreeAddress[Idx] < BTreeAddress[Idx] )
+			{
+				// A has a lower index, show it first:
+				return true;
+			}
+			else if( ATreeAddress[Idx] > BTreeAddress[Idx] )
+			{
+				// B has a lower index, show it first:
+				return false;
+			}
+			else
+			{
+				// tie, go to the next level of the tree:
+				++Idx;
+			}
+		}
+
+		// fall back, just let diff type win:
+		return A.DiffType < B.DiffType;
+	};
+
+
+	DifferingProperties.Entries.Sort(SortTreePredicate);
+
 	ChildSlot[
 		SNew(SSplitter)
 			+ SSplitter::Slot()
@@ -107,19 +153,61 @@ void SMergeTreeView::PrevDiff()
 
 bool SMergeTreeView::HasNextDifference() const
 {
-	return DifferingProperties.IsValidIndex(CurrentDifference + 1);
+	return DifferingProperties.Entries.IsValidIndex(CurrentDifference + 1);
 }
 
 bool SMergeTreeView::HasPrevDifference() const
 {
-	return DifferingProperties.IsValidIndex(CurrentDifference - 1);
+	return DifferingProperties.Entries.IsValidIndex(CurrentDifference - 1);
+}
+
+void SMergeTreeView::HighlightNextConflict()
+{
+	if (CurrentMergeConflict + 1 < MergeConflicts.Entries.Num())
+	{
+		++CurrentMergeConflict;
+	}
+	HighlightCurrentConflict();
+}
+
+void SMergeTreeView::HighlightPrevConflict()
+{
+	if (CurrentMergeConflict - 1 >= 0)
+	{
+		--CurrentMergeConflict;
+	}
+	HighlightCurrentConflict();
+}
+
+bool SMergeTreeView::HasNextConflict() const
+{
+	// return true if we have one conflict so that users can reselect the conflict if they desire. If we 
+	// return false when we already have selected this one and only conflict then there will be no way
+	// to reselect it if the user wants to.
+	return MergeConflicts.Entries.Num() != 0 && (MergeConflicts.Entries.Num() == 1 || CurrentMergeConflict + 1 < MergeConflicts.Entries.Num());
+}
+
+bool SMergeTreeView::HasPrevConflict() const
+{
+	// note in HasNextConflict applies here as well.
+	return MergeConflicts.Entries.Num() == 1 || CurrentMergeConflict > 0;
+}
+
+void SMergeTreeView::HighlightCurrentConflict()
+{
+	const FSCSDiffEntry& Conflict = MergeConflicts.Entries[CurrentMergeConflict];
+	for (auto& View : SCSViews)
+	{
+		View.HighlightProperty(Conflict.TreeIdentifier.Name, Conflict.PropertyIdentifier);
+	}
 }
 
 void SMergeTreeView::HighlightCurrentDifference()
 {
+	const FSCSDiffEntry& Difference = DifferingProperties.Entries[CurrentDifference];
 	for (auto& View : SCSViews)
 	{
-		View.HighlightProperty(DifferingProperties[CurrentDifference]);
+		View.HighlightProperty(Difference.TreeIdentifier.Name, Difference.PropertyIdentifier);
 	}
 }
 
diff --git a/Engine/Source/Developer/Merge/Private/SMergeTreeView.h b/Engine/Source/Developer/Merge/Private/SMergeTreeView.h
--- a/Engine/Source/Developer/Merge/Private/SMergeTreeView.h
+++ b/Engine/Source/Developer/Merge/Private/SMergeTreeView.h
@@ -7,6 +7,7 @@
 
 class SMergeTreeView	: public SCompoundWidget
 						, public IDiffControl
+						, public IMergeControl
 {
 public:
 	virtual ~SMergeTreeView() {}
@@ -23,6 +24,13 @@ private:
 	bool HasNextDifference() const override;
 	bool HasPrevDifference() const override;
 
+	/** Implementation of IMergeControl: */
+	void HighlightNextConflict() override;
+	void HighlightPrevConflict() override;
+	bool HasNextConflict() const override;
+	bool HasPrevConflict() const override;
+
+	void HighlightCurrentConflict();
 	void HighlightCurrentDifference();
 	FSCSDiff& GetRemoteView();
 	FSCSDiff& GetBaseView();
@@ -31,6 +39,9 @@ private:
 	FBlueprintMergeData Data;
 	TArray< FSCSDiff > SCSViews;
 
-	TArray< FSCSDiffEntry > DifferingProperties;
+	FSCSDiffRoot MergeConflicts;
+	int CurrentMergeConflict;
+
+	FSCSDiffRoot DifferingProperties;
 	int CurrentDifference;
 };
diff --git a/Engine/Source/Editor/Kismet/Private/DiffUtils.cpp b/Engine/Source/Editor/Kismet/Private/DiffUtils.cpp
--- a/Engine/Source/Editor/Kismet/Private/DiffUtils.cpp
+++ b/Engine/Source/Editor/Kismet/Private/DiffUtils.cpp
@@ -103,94 +103,60 @@ void DiffUtils::CompareUnrelatedObjects(const UObject* A, UObject const* B, TArr
 	}
 }
 
-void DiffUtils::CompareUnrelatedSCS(const UBlueprint* A, const UBlueprint* B, TArray< FSCSDiffEntry >& OutDifferingEntries, TArray< int >* OptionalOutSortKeys)
+void DiffUtils::CompareUnrelatedSCS(const UBlueprint* Old, const TArray< FSCSResolvedIdentifier >& OldHierarchy, const UBlueprint* New, const TArray< FSCSResolvedIdentifier >& NewHierarchy, FSCSDiffRoot& OutDifferingEntries )
 {
-	const auto FindObjectPropertyByName = [](const UObject* Instance, FName Name) -> UObject const*
+	const auto FindEntry = [](TArray< FSCSResolvedIdentifier > const& InArray, const FSCSIdentifier* Value) -> const FSCSResolvedIdentifier*
 	{
-		for (TFieldIterator<UProperty> PropertyIter(Instance->GetClass()); PropertyIter; ++PropertyIter)
+		for (const auto& Node : InArray)
 		{
-			if (PropertyIter->GetFName() == Name)
+			if (Node.Identifier.Name == Value->Name )
 			{
-				const UObject* const* SCS = PropertyIter->ContainerPtrToValuePtr<const UObject*>(Instance);
-				return SCS ? *SCS : nullptr;
+				return &Node;
 			}
 		}
 		return nullptr;
 	};
 
-	FName SCS = TEXT("SimpleConstructionScript");
-	UObject const* SCSA = FindObjectPropertyByName(A, SCS);
-	UObject const* SCSB = FindObjectPropertyByName(B, SCS);
-
-	// we get bitten by const shallowness here, avoid mutation:
-	const TArray<USCS_Node*> NodesInA = A && A->SimpleConstructionScript ? A->SimpleConstructionScript->GetAllNodes() : TArray<USCS_Node*>();
-	const TArray<USCS_Node*> NodesInB = B && B->SimpleConstructionScript ? B->SimpleConstructionScript->GetAllNodes() : TArray<USCS_Node*>();
-
-	const auto FindByName = [](const TArray<USCS_Node*>& NodeList, FName Name) -> USCS_Node*
+	for (const auto& OldNode : OldHierarchy)
 	{
-		for (auto Node : NodeList)
+		const FSCSResolvedIdentifier* NewEntry = FindEntry(NewHierarchy, &OldNode.Identifier);
+
+		if (NewEntry != nullptr)
 		{
-			if (Node->VariableName == Name)
+			// @todo doc: did properties change?
+			TArray<FPropertySoftPath> DifferingProperties;
+			DiffUtils::CompareUnrelatedObjects(OldNode.Object, NewEntry->Object, DifferingProperties);
+			for (const auto& Property : DifferingProperties)
 			{
-				return Node;
+				FSCSDiffEntry Diff = { OldNode.Identifier, ETreeDiffType::NODE_PROPERTY_CHANGED, Property };
+				OutDifferingEntries.Entries.Push(Diff);
 			}
-		}
-		return nullptr;
-	};
 
-	int32 SortKey = 0;
-	for (auto NodeInA : NodesInA)
-	{
-		auto NodeInB = FindByName(NodesInB, NodeInA->VariableName);
-		if (NodeInB)
-		{
-			// diff NodeA vs NodeB:
-			TArray<FPropertySoftPath> DifferingProperties;
+			// did it move?
 
-			DiffUtils::CompareUnrelatedObjects(A, B, DifferingProperties);
-			for (auto Property : DifferingProperties)
-			{
-				FSCSDiffEntry Diff = { NodeInA->VariableName, Property.LastPropertyName() };
-				OutDifferingEntries.Push(Diff);
 
-				if (OptionalOutSortKeys)
-				{
-					OptionalOutSortKeys->Push(SortKey);
-				}
-			}
+			// no change! Do nothing.
 		}
 		else
 		{
-			// Node A was added (removed from node B):
-			FSCSDiffEntry Diff = { NodeInA->VariableName, FName() };
-			OutDifferingEntries.Push(Diff);
-
-			if (OptionalOutSortKeys)
-			{
-				OptionalOutSortKeys->Push(SortKey);
-			}
+			// not found in the new data, must have been deleted:
+			FSCSDiffEntry Entry = { OldNode.Identifier, ETreeDiffType::NODE_REMOVED, FPropertySoftPath() };
+			OutDifferingEntries.Entries.Push( Entry );
 		}
-
-		++SortKey;
 	}
 
-	// add nodes that were added in B:
-	for (auto NodeInB : NodesInB)
+	for (const auto& NewNode : NewHierarchy)
 	{
-		if (FindByName(NodesInA, NodeInB->VariableName) == nullptr)
-		{
-			// Node B was added:
-			FSCSDiffEntry Diff = { NodeInB->VariableName, FName() };
-			OutDifferingEntries.Push(Diff);
+		const FSCSResolvedIdentifier* OldEntry = FindEntry(OldHierarchy, &NewNode.Identifier);
 
-			if (OptionalOutSortKeys)
-			{
-				OptionalOutSortKeys->Push(SortKey);
-			}
+		if (OldEntry == nullptr)
+		{
+			FSCSDiffEntry Entry = { NewNode.Identifier, ETreeDiffType::NODE_ADDED, FPropertySoftPath() };
+			OutDifferingEntries.Entries.Push( Entry );
 		}
-
-		++SortKey;
 	}
+
+
 }
 
 bool DiffUtils::Identical(const FResolvedProperty& AProp, const FResolvedProperty& BProp)
diff --git a/Engine/Source/Editor/Kismet/Private/SBlueprintDiff.cpp b/Engine/Source/Editor/Kismet/Private/SBlueprintDiff.cpp
--- a/Engine/Source/Editor/Kismet/Private/SBlueprintDiff.cpp
+++ b/Engine/Source/Editor/Kismet/Private/SBlueprintDiff.cpp
@@ -36,7 +36,6 @@ public:
 	FSCSDiffControl(
 					const UBlueprint* InOldBlueprint
 					, const UBlueprint* InNewBlueprint
-					, TArray< FSCSDiffEntry > InDifferingProperties
 					);
 
 
@@ -52,7 +51,7 @@ private:
 
 	void HighlightCurrentDifference();
 
-	TArray< FSCSDiffEntry > DifferingProperties;
+	FSCSDiffRoot DifferingProperties;
 	int CurrentDifference;
 
 	FSCSDiff OldSCS;
@@ -62,13 +61,15 @@ private:
 FSCSDiffControl::FSCSDiffControl(
 		const UBlueprint* InOldBlueprint
 		, const UBlueprint* InNewBlueprint
-		, TArray< FSCSDiffEntry > InDifferingProperties
 	)
-	: DifferingProperties(InDifferingProperties)
+	: DifferingProperties()
 	, CurrentDifference(-1)
 	, OldSCS( InOldBlueprint )
 	, NewSCS( InNewBlueprint )
 {
+	TArray< FSCSResolvedIdentifier > OldHierarchy = OldSCS.GetDisplayedHierarchy();
+	TArray< FSCSResolvedIdentifier > NewHierarchy = NewSCS.GetDisplayedHierarchy();
+	DiffUtils::CompareUnrelatedSCS(InOldBlueprint, OldHierarchy, InNewBlueprint, NewHierarchy, DifferingProperties);
 }
 
 void FSCSDiffControl::NextDiff()
@@ -85,18 +86,18 @@ void FSCSDiffControl::PrevDiff()
 
 bool FSCSDiffControl::HasNextDifference() const
 {
-	return DifferingProperties.IsValidIndex(CurrentDifference + 1);
+	return DifferingProperties.Entries.IsValidIndex(CurrentDifference + 1);
 }
 
 bool FSCSDiffControl::HasPrevDifference() const
 {
-	return DifferingProperties.IsValidIndex(CurrentDifference - 1);
+	return DifferingProperties.Entries.IsValidIndex(CurrentDifference - 1);
 }
 
 void FSCSDiffControl::HighlightCurrentDifference()
 {
-	OldSCS.HighlightProperty(DifferingProperties[CurrentDifference]);
-	NewSCS.HighlightProperty(DifferingProperties[CurrentDifference]);
+	OldSCS.HighlightProperty(DifferingProperties.Entries[CurrentDifference].TreeIdentifier.Name, FPropertyPath());
+	NewSCS.HighlightProperty(DifferingProperties.Entries[CurrentDifference].TreeIdentifier.Name, FPropertyPath());
 }
 
 class FCDODiffControl	: public TSharedFromThis<FCDODiffControl>
@@ -1284,10 +1285,7 @@ TSharedRef<SWidget> SBlueprintDiff::GenerateDefaultsPanel()
 TSharedRef<SWidget> SBlueprintDiff::GenerateComponentsPanel()
 {
 	//Splitter for left and right blueprint. Current convention is for the local (probably newer?) blueprint to be on the right:
-	TArray< FSCSDiffEntry > Differences;
-	DiffUtils::CompareUnrelatedSCS( PanelOld.Blueprint, PanelNew.Blueprint, Differences );
-
-	auto NewDiffControl = TSharedPtr<FSCSDiffControl>(new FSCSDiffControl(PanelOld.Blueprint, PanelNew.Blueprint, Differences));
+	auto NewDiffControl = TSharedPtr<FSCSDiffControl>(new FSCSDiffControl(PanelOld.Blueprint, PanelNew.Blueprint) );
 	DiffControl = NewDiffControl;
 	return SNew(SSplitter)
 		+ SSplitter::Slot()
diff --git a/Engine/Source/Editor/Kismet/Private/SCSDiff.cpp b/Engine/Source/Editor/Kismet/Private/SCSDiff.cpp
--- a/Engine/Source/Editor/Kismet/Private/SCSDiff.cpp
+++ b/Engine/Source/Editor/Kismet/Private/SCSDiff.cpp
@@ -25,26 +25,42 @@ FSCSDiff::FSCSDiff(const UBlueprint* InBlueprint)
 		];
 }
 
-void FSCSDiff::HighlightProperty(FSCSDiffEntry Property)
+void FSCSDiff::HighlightProperty(FName VarName, FPropertySoftPath Property)
 {
-	check( Property.TreeNodeName != FName() );
-	if( const USimpleConstructionScript* SCS = SCSEditor->SCS )
-	{
-		TArray<USCS_Node*> Nodes = SCS->GetAllNodes();
-		for( auto Node : Nodes )
-		{
-			if( Node->VariableName == Property.TreeNodeName )
-			{
-				SCSEditor->HighlightTreeNode( Node, Property.PropertyName );
-				return;
-			}
-		}
-	}
-
-	SCSEditor->ClearSelection();
+	check( VarName != FName() );
+	SCSEditor->HighlightTreeNode( VarName, FPropertyPath() );
 }
 
 TSharedRef< SWidget > FSCSDiff::TreeWidget()
 {
 	return ContainerWidget.ToSharedRef();
 }
+
+void GetDisplayedHierarchyRecursive(TArray< int32 >& TreeAddress, const FSCSEditorTreeNode& Node, TArray< FSCSResolvedIdentifier >& OutResult)
+{
+	FSCSIdentifier Identifier = { Node.GetVariableName(), TreeAddress };
+	FSCSResolvedIdentifier ResolvedIdentifier = { Identifier, Node.GetComponentTemplate() };
+	OutResult.Push(ResolvedIdentifier);
+	const auto& Children = Node.GetChildren();
+	for (int32 Iter = 0; Iter != Children.Num(); ++Iter)
+	{
+		TreeAddress.Push(Iter);
+		GetDisplayedHierarchyRecursive(TreeAddress, *Children[Iter], OutResult);
+		TreeAddress.Pop();
+	}
+}
+
+TArray< FSCSResolvedIdentifier > FSCSDiff::GetDisplayedHierarchy() const
+{
+	TArray< FSCSResolvedIdentifier > Ret;
+
+	for (int32 Iter = 0; Iter != SCSEditor->RootNodes.Num(); ++Iter)
+	{
+		TArray< int32 > TreeAddress;
+		TreeAddress.Push(Iter);
+		GetDisplayedHierarchyRecursive(TreeAddress, *SCSEditor->RootNodes[Iter], Ret);
+	}
+
+	return Ret;
+}
+
diff --git a/Engine/Source/Editor/Kismet/Private/SSCSEditor.cpp b/Engine/Source/Editor/Kismet/Private/SSCSEditor.cpp
--- a/Engine/Source/Editor/Kismet/Private/SSCSEditor.cpp
+++ b/Engine/Source/Editor/Kismet/Private/SSCSEditor.cpp
@@ -2481,6 +2481,48 @@ void SSCSEditor::SelectNode(FSCSEditorTreeNodePtrType InNodeToSelect, bool IsCnt
 	}
 }
 
+static FSCSEditorTreeNode* FindRecursive( FSCSEditorTreeNode* Node, FName Name )
+{
+	if (Node->GetVariableName() == Name)
+	{
+		return Node;
+	}
+	else
+	{
+		for (const auto& Child : Node->GetChildren())
+		{
+			if (auto Result = FindRecursive(Child.Get(), Name))
+			{
+				return Result;
+			}
+		}
+	}
+
+	return nullptr;
+}
+
+void SSCSEditor::HighlightTreeNode(FName TreeNodeName, const FPropertyPath& Property)
+{
+	for( const auto& Node : RootNodes )
+	{
+		if( auto FoundNode = FindRecursive( Node.Get(), TreeNodeName ) )
+		{
+			SelectNode(FoundNode->AsShared(), false);
+
+			if (Property != FPropertyPath())
+			{
+				auto KismetInspectorSPtr = KismetInspectorPtr.Pin();
+				check(KismetInspectorSPtr.IsValid());
+				KismetInspectorSPtr->GetPropertyView()->HighlightProperty(Property);
+			}
+
+			return;
+		}
+	}
+	
+	ClearSelection();
+}
+
 void SSCSEditor::HighlightTreeNode(const USCS_Node* Node, FName Property)
 {
 	check(Node);
diff --git a/Engine/Source/Editor/Kismet/Private/SSCSEditor.h b/Engine/Source/Editor/Kismet/Private/SSCSEditor.h
--- a/Engine/Source/Editor/Kismet/Private/SSCSEditor.h
+++ b/Engine/Source/Editor/Kismet/Private/SSCSEditor.h
@@ -467,7 +467,15 @@ public:
 
 	/** Select the given tree node */
 	void SelectNode(FSCSEditorTreeNodePtrType InNodeToSelect, bool IsCntrlDown);
-	
+
+	/**
+	 * Highlight a tree node and, optionally, a property with in it
+	 *
+	 * @param TreeNodeName		Name of the treenode to be highlighted
+	 * @param Property	The name of the property to be highlighted in the details view
+	 * @return True if the node was found in this Editor, otherwise false
+	 */
+	void HighlightTreeNode( FName TreeNodeName, const class FPropertyPath& Property );
 	/**
 	 * Highlight a tree node and, optionally, a property with in it
 	 *
diff --git a/Engine/Source/Editor/Kismet/Public/DiffUtils.h b/Engine/Source/Editor/Kismet/Public/DiffUtils.h
--- a/Engine/Source/Editor/Kismet/Public/DiffUtils.h
+++ b/Engine/Source/Editor/Kismet/Public/DiffUtils.h
@@ -35,6 +35,10 @@ struct FResolvedProperty
  */
 struct FPropertySoftPath
 {
+	FPropertySoftPath()
+	{
+	}
+
 	FPropertySoftPath(TArray<FName> InPropertyChain)
 		: PropertyChain(InPropertyChain)
 	{
@@ -56,11 +60,38 @@ struct FPropertySoftPath
 	{
 		return PropertyChain == RHS.PropertyChain;
 	}
+
+	inline bool operator!=(FPropertySoftPath const& RHS ) const
+	{
+		return !(*this == RHS);
+	}
 private:
 	friend uint32 GetTypeHash( FPropertySoftPath const& Path );
 	TArray<FName> PropertyChain;
 };
 
+struct FSCSIdentifier
+{
+	FName Name;
+	TArray< int32 > TreeLocation;
+};
+
+struct FSCSResolvedIdentifier
+{
+	FSCSIdentifier Identifier;
+	const UObject* Object;
+};
+
+FORCEINLINE bool operator==( const FSCSIdentifier& A, const FSCSIdentifier& B )
+{
+	return A.Name == B.Name && A.TreeLocation == B.TreeLocation;
+}
+
+FORCEINLINE bool operator!=(const FSCSIdentifier& A, const FSCSIdentifier& B)
+{
+	return !(A == B);
+}
+
 FORCEINLINE FName FPropertySoftPath::LastPropertyName() const
 {
 	if( PropertyChain.Num() != 0 )
@@ -83,22 +114,37 @@ FORCEINLINE uint32 GetTypeHash( FPropertySoftPath const& Path )
 // Trying to restrict us to this typedef because I'm a little skeptical about hashing FPropertySoftPath safely
 typedef TSet< FPropertySoftPath > FPropertySoftPathSet;
 
+namespace ETreeDiffType
+{
+	enum Type
+	{
+		UNKNOWN,
+		NODE_ADDED,
+		NODE_REMOVED,
+		NODE_PROPERTY_CHANGED,
+		NODE_MOVED,
+		/** We could potentially try to identify hierarchy reorders separately from add/remove */
+	};
+}
+
 struct FSCSDiffEntry
 {
-	FName TreeNodeName;
-	FName PropertyName;
+	FSCSIdentifier TreeIdentifier;
+	ETreeDiffType::Type DiffType;
+	FPropertySoftPath PropertyIdentifier;
 };
 
-inline bool operator ==(const FSCSDiffEntry A, const FSCSDiffEntry B)
+struct FSCSDiffRoot
 {
-	return A.TreeNodeName == B.TreeNodeName && A.PropertyName == B.PropertyName;
-}
+	// use indices in FSCSIdentifier::TreeLocation to find hierarchy..
+	TArray< FSCSDiffEntry > Entries;
+};
 
 namespace DiffUtils
 {
 	KISMET_API const UObject* GetCDO(const UBlueprint* ForBlueprint);
 	KISMET_API void CompareUnrelatedObjects(const UObject* A, const UObject* B, TArray<FPropertySoftPath>& OutDifferingProperties);
-	KISMET_API void CompareUnrelatedSCS(const UBlueprint* Old, const UBlueprint* New, TArray< FSCSDiffEntry >& OutDifferingEntries, TArray< int >* OptionalOutSortKeys = nullptr);
+	KISMET_API void CompareUnrelatedSCS(const UBlueprint* Old, const TArray< FSCSResolvedIdentifier >& OldHierarchy, const UBlueprint* New, const TArray< FSCSResolvedIdentifier >& NewHierarchy, FSCSDiffRoot& OutDifferingEntries );
 	KISMET_API bool Identical(const FResolvedProperty& AProp, const FResolvedProperty& BProp);
 	TArray<FPropertySoftPath> GetVisiblePropertiesInOrderDeclared(const UObject* ForObj, const TArray<FName>& Scope = TArray<FName>());
 
diff --git a/Engine/Source/Editor/Kismet/Public/SCSDiff.h b/Engine/Source/Editor/Kismet/Public/SCSDiff.h
--- a/Engine/Source/Editor/Kismet/Public/SCSDiff.h
+++ b/Engine/Source/Editor/Kismet/Public/SCSDiff.h
@@ -8,8 +8,10 @@ class KISMET_API FSCSDiff
 public:
 	FSCSDiff(const class UBlueprint* InBlueprint);
 
-	void HighlightProperty(FSCSDiffEntry Property);
+	void HighlightProperty(FName VarName, FPropertySoftPath Property);
 	TSharedRef< SWidget > TreeWidget();
+
+	TArray< FSCSResolvedIdentifier > GetDisplayedHierarchy() const;
 private:
 	TSharedPtr< class SWidget > ContainerWidget;
 	TSharedPtr< class SSCSEditor > SCSEditor;
