Fixed graph view not allowing you to iterate differences when there were no local or no remote changes

diff --git a/Engine/Source/Developer/Merge/Private/SMergeGraphView.cpp b/Engine/Source/Developer/Merge/Private/SMergeGraphView.cpp
--- a/Engine/Source/Developer/Merge/Private/SMergeGraphView.cpp
+++ b/Engine/Source/Developer/Merge/Private/SMergeGraphView.cpp
@@ -593,7 +593,7 @@ bool SMergeGraphView::HasNoDifferences() const
 {
 	auto RemoteDiffResultListPtr = RemoteDiffResultList.Pin();
 	auto LocalDiffResultListPtr = LocalDiffResultList.Pin();
-	return	(!RemoteDiffResultListPtr.IsValid() || RemoteDiffResultListPtr->GetNumItemsBeingObserved() == 0) ||
+	return	(!RemoteDiffResultListPtr.IsValid() || RemoteDiffResultListPtr->GetNumItemsBeingObserved() == 0) &&
 			(!LocalDiffResultListPtr.IsValid() || LocalDiffResultListPtr->GetNumItemsBeingObserved() == 0);
 }
 
