Fix merge menu entry showing up in the same menu in two places.

diff --git a/Engine/Source/Developer/Merge/Private/Merge.cpp b/Engine/Source/Developer/Merge/Private/Merge.cpp
--- a/Engine/Source/Developer/Merge/Private/Merge.cpp
+++ b/Engine/Source/Developer/Merge/Private/Merge.cpp
@@ -187,12 +187,12 @@ void FMerge::StartupModule()
 {
 	// This code will execute after your module is loaded into memory (but after global variables are initialized, of course.)
 
-	FGlobalTabmanager::Get()->RegisterNomadTabSpawner(
-		MergeToolTabId, 
-		FOnSpawnTab::CreateStatic( [](const FSpawnTabArgs&) { return SNew(SDockTab); } ) )
+
+	// Registering a nomad spawner that spawns an empty dock tab on purpose - allows us to call InvokeTab() using our TabId later and set the content. (see GenerateMergeWidget())
+	FGlobalTabmanager::Get()->RegisterNomadTabSpawner(MergeToolTabId, FOnSpawnTab::CreateStatic([] (const FSpawnTabArgs&) { return SNew(SDockTab); }))
 		.SetDisplayName(NSLOCTEXT("MergeTool", "TabTitle", "Merge Tool"))
-		.SetTooltipText(NSLOCTEXT("MergeTool", "TooltipText", "Used to display several versions of a blueprint that need to be merged into a single version.")
-	);
+		.SetTooltipText(NSLOCTEXT("MergeTool", "TooltipText", "Used to display several versions of a blueprint that need to be merged into a single version."))
+		.SetAutoGenerateMenuEntry(false);
 }
 
 void FMerge::ShutdownModule()
