[UE-5744/UE-5743] Adding a mode to the blueprint merge panel that lets you pick arbitrary blueprints to merge against each other.

diff --git a/Engine/Source/Developer/Merge/Private/BlueprintMergeData.h b/Engine/Source/Developer/Merge/Private/BlueprintMergeData.h
--- a/Engine/Source/Developer/Merge/Private/BlueprintMergeData.h
+++ b/Engine/Source/Developer/Merge/Private/BlueprintMergeData.h
@@ -10,10 +10,11 @@ struct FBlueprintMergeData
 	FBlueprintMergeData()
 		: OwningEditor()
 		, BlueprintLocal(NULL)
+		, RevisionLocal(FRevisionInfo::InvalidRevision())
 		, BlueprintBase(NULL)
-		, RevisionBase()
+		, RevisionBase(FRevisionInfo::InvalidRevision())
 		, BlueprintRemote(NULL)
-		, RevisionRemote()
+		, RevisionRemote(FRevisionInfo::InvalidRevision())
 	{
 	}
 
@@ -37,6 +38,8 @@ struct FBlueprintMergeData
 	TWeakPtr<class FBlueprintEditor>	OwningEditor;
 	
 	const class UBlueprint*				BlueprintLocal;
+	FRevisionInfo						RevisionLocal;
+
 	const class UBlueprint*				BlueprintBase;
 	FRevisionInfo						RevisionBase;
 
diff --git a/Engine/Source/Developer/Merge/Private/Merge.cpp b/Engine/Source/Developer/Merge/Private/Merge.cpp
--- a/Engine/Source/Developer/Merge/Private/Merge.cpp
+++ b/Engine/Source/Developer/Merge/Private/Merge.cpp
@@ -10,6 +10,7 @@
 #include "SDockTab.h"
 #include "SNotificationList.h"
 #include "NotificationManager.h"
+#include "MergeUtils.h"
 
 #define LOCTEXT_NAMESPACE "Merge"
 
@@ -22,97 +23,27 @@ static void DisplayErrorMessage( const FText& ErrorMessage )
 	FSlateNotificationManager::Get().AddNotification(Info);
 }
 
-static FSourceControlStatePtr GetSourceControlState(const FString& PackageName)
-{
-	ISourceControlProvider& SourceControlProvider = ISourceControlModule::Get().GetProvider();
-	TSharedRef<FUpdateStatus, ESPMode::ThreadSafe> UpdateStatusOperation = ISourceControlOperation::Create<FUpdateStatus>();
-	UpdateStatusOperation->SetUpdateHistory(true);
-	SourceControlProvider.Execute(UpdateStatusOperation, SourceControlHelpers::PackageFilename(PackageName));
-
-	FSourceControlStatePtr State = SourceControlProvider.GetState(SourceControlHelpers::PackageFilename(PackageName), EStateCacheUsage::Use);
-	if (!State.IsValid() || !State->IsSourceControlled() || !FPackageName::DoesPackageExist(PackageName))
-	{
-		return FSourceControlStatePtr();
-	}
-	else
-	{
-		return State;
-	}
-}
-
-static UObject* LoadRevision(const FString& AssetName, const ISourceControlRevision& DesiredRevision)
-{
-	// Get the head revision of this package from source control
-	FString TempFileName;
-	if (DesiredRevision.Get(TempFileName))
-	{
-		// Forcibly disable compile on load in case we are loading old blueprints that might try to update/compile
-		TGuardValue<bool> DisableCompileOnLoad(GForceDisableBlueprintCompileOnLoad, true);
-
-		// Try and load that package
-		UPackage* TempPackage = LoadPackage(NULL, *TempFileName, LOAD_None);
-		if (TempPackage != NULL)
-		{
-			// Grab the old asset from that old package
-			UObject* OldObject = FindObject<UObject>(TempPackage, *AssetName);
-			if (OldObject)
-			{
-				return OldObject;
-			}
-			else
-			{
-				DisplayErrorMessage( 
-					FText::Format(
-						LOCTEXT("MergedFailedToFindObject", "Aborted Load of {0} because we could not find an object named {1}" )
-						, FText::FromString(TempFileName)
-						, FText::FromString(AssetName) 
-					)
-				);
-			}
-		}
-		else
-		{
-			DisplayErrorMessage(
-				FText::Format(
-					LOCTEXT("MergedFailedToLoadPackage", "Aborted Load of {0} because we could not load the package")
-					, FText::FromString(TempFileName)
-				)
-			);
-		}
-	}
-	else
-	{
-		DisplayErrorMessage(
-			FText::Format(
-				LOCTEXT("MergedFailedToFindRevision", "Aborted Load of {0} because we could not get the requested revision")
-				, FText::FromString(TempFileName)
-			)
-		);
-	}
-	return NULL;
-}
-
 static FRevisionInfo GetRevisionInfo(ISourceControlRevision const& FromRevision)
 {
 	FRevisionInfo Ret = { FromRevision.GetRevisionNumber(), FromRevision.GetCheckInIdentifier(), FromRevision.GetDate() };
 	return Ret;
 }
 
-static UObject* LoadHeadRev(const FString& PackageName, const FString& AssetName, const ISourceControlState& SourceControlState, FRevisionInfo& OutRevInfo)
+static const UObject* LoadHeadRev(const FString& PackageName, const FString& AssetName, const ISourceControlState& SourceControlState, FRevisionInfo& OutRevInfo)
 {
 	// HistoryItem(0) is apparently the head revision:
 	TSharedPtr<ISourceControlRevision, ESPMode::ThreadSafe> Revision = SourceControlState.GetHistoryItem(0);
 	check(Revision.IsValid());
 	OutRevInfo = GetRevisionInfo(*Revision);
-	return LoadRevision(AssetName, *Revision);
+	return FMergeToolUtils::LoadRevision(AssetName, *Revision);
 }
 
-static UObject* LoadBaseRev(const FString& PackageName, const FString& AssetName, const ISourceControlState& SourceControlState, FRevisionInfo& OutRevInfo)
+static const UObject* LoadBaseRev(const FString& PackageName, const FString& AssetName, const ISourceControlState& SourceControlState, FRevisionInfo& OutRevInfo)
 {
 	TSharedPtr<ISourceControlRevision, ESPMode::ThreadSafe> Revision = SourceControlState.GetBaseRevForMerge();
 	check(Revision.IsValid());
 	OutRevInfo = GetRevisionInfo(*Revision);
-	return LoadRevision(AssetName, *Revision);
+	return FMergeToolUtils::LoadRevision(AssetName, *Revision);
 }
 
 static TSharedPtr<SWidget> GenerateMergeTabContents(TSharedRef<FBlueprintEditor> Editor,
@@ -123,51 +54,31 @@ static TSharedPtr<SWidget> GenerateMergeTabContents(TSharedRef<FBlueprintEditor>
 													const UBlueprint*       LocalBlueprint,
 													const FOnMergeResolved& MereResolutionCallback)
 {
-	TSharedPtr<SWidget> TabContent;
-
-	if (RemoteBlueprint && BaseBlueprint)
+	bool bForceAssetPicker = false;
+	if (BaseBlueprint == nullptr)
 	{
-		FBlueprintMergeData Data(Editor
-			, LocalBlueprint
-			, BaseBlueprint
-			, RemoteRevInfo
-			, RemoteBlueprint
-			, BaseRevInfo);
-
-		TabContent = SNew(SBlueprintMerge, Data)
-			.OnMergeResolved(MereResolutionCallback);
+		BaseBlueprint = LocalBlueprint;
+		bForceAssetPicker = true;
 	}
-	else
+	if (RemoteBlueprint == nullptr)
 	{
-		FText MissingFiles;
-		if (!RemoteBlueprint && !BaseBlueprint)
-		{
-			MissingFiles = LOCTEXT("MergeBothRevisionsFailed", "common base, nor conflicting revision");
-		}
-		else if (!RemoteBlueprint)
-		{
-			MissingFiles = LOCTEXT("MergeConflictingRevisionsFailed", "conflicting revision");
-		}
-		else
-		{
-			MissingFiles = LOCTEXT("MergeBaseRevisionsFailed", "common base");
-		}
-
-		DisplayErrorMessage(
-			FText::Format(
-				LOCTEXT("MergeRevisionLoadFailed", "Aborted Merge of {0} because we could not load {1}")
-				, FText::FromString(LocalBlueprint->GetName())
-				, MissingFiles
-			)
-		);
-
-		TabContent = SNew(SHorizontalBox);
+		RemoteBlueprint = LocalBlueprint;
+		bForceAssetPicker = true;
 	}
 
-	return TabContent;
+	FBlueprintMergeData Data(Editor
+		, LocalBlueprint
+		, BaseBlueprint
+		, RemoteRevInfo
+		, RemoteBlueprint
+		, BaseRevInfo);
+
+	return SNew(SBlueprintMerge, Data)
+		.bForcePickAssets(bForceAssetPicker)
+		.OnMergeResolved(MereResolutionCallback);
 }
 
-class FMerge : public IMerge
+class FMerge : public IMerge 
 {
 	/** IModuleInterface implementation */
 	virtual void StartupModule() override;
@@ -216,35 +127,44 @@ TSharedRef<SDockTab> FMerge::GenerateMergeWidget(const UBlueprint& Object, TShar
 	// merge the local asset with the depot, SCC provides us with the last common revision as
 	// a basis for the merge:
 
-	// @todo DO: this will probably need to be async.. pulling down some old versions of assets:
-	const FString& PackageName = Object.GetOutermost()->GetName();
-	const FString& AssetName = Object.GetName();
-	
 	TSharedPtr<SWidget> Contents;
 
-	FSourceControlStatePtr SourceControlState = GetSourceControlState(PackageName);
-	if (!SourceControlState.IsValid())
+	if (!PendingMerge(Object))
 	{
-		DisplayErrorMessage(
-			FText::Format(
-				LOCTEXT("MergeFailedNoSourceControl", "Aborted Load of {0} from {1} because the source control state was invalidated")
-				, FText::FromString(AssetName)
-				, FText::FromString(PackageName)
-			)
-		);
-
-		Contents = SNew(SHorizontalBox);
+		// this should load up the merge-tool, with an asset picker, where they
+		// can pick the asset/revisions to merge against
+		Contents = GenerateMergeTabContents(Editor, nullptr, FRevisionInfo::InvalidRevision(), nullptr, FRevisionInfo::InvalidRevision(), &Object, FOnMergeResolved());
 	}
 	else
 	{
-		ISourceControlState const& SourceControlStateRef = *SourceControlState;
+		// @todo DO: this will probably need to be async.. pulling down some old versions of assets:
+		const FString& PackageName = Object.GetOutermost()->GetName();
+		const FString& AssetName = Object.GetName();
+
+		FSourceControlStatePtr SourceControlState = FMergeToolUtils::GetSourceControlState(PackageName);
+		if (!SourceControlState.IsValid())
+		{
+			DisplayErrorMessage(
+				FText::Format(
+					LOCTEXT("MergeFailedNoSourceControl", "Aborted Load of {0} from {1} because the source control state was invalidated")
+					, FText::FromString(AssetName)
+					, FText::FromString(PackageName)
+				)
+			);
+
+			Contents = SNew(SHorizontalBox);
+		}
+		else
+		{
+			ISourceControlState const& SourceControlStateRef = *SourceControlState;
 
-		FRevisionInfo CurrentRevInfo = FRevisionInfo::InvalidRevision();
-		const UBlueprint* RemoteBlueprint = Cast< UBlueprint >(LoadHeadRev(PackageName, AssetName, SourceControlStateRef, CurrentRevInfo));
-		FRevisionInfo BaseRevInfo = FRevisionInfo::InvalidRevision();
-		const UBlueprint* BaseBlueprint = Cast< UBlueprint >(LoadBaseRev(PackageName, AssetName, SourceControlStateRef, BaseRevInfo));
+			FRevisionInfo CurrentRevInfo = FRevisionInfo::InvalidRevision();
+			const UBlueprint* RemoteBlueprint = Cast< UBlueprint >(LoadHeadRev(PackageName, AssetName, SourceControlStateRef, CurrentRevInfo));
+			FRevisionInfo BaseRevInfo = FRevisionInfo::InvalidRevision();
+			const UBlueprint* BaseBlueprint = Cast< UBlueprint >(LoadBaseRev(PackageName, AssetName, SourceControlStateRef, BaseRevInfo));
 
-		Contents = GenerateMergeTabContents(Editor, BaseBlueprint, BaseRevInfo, RemoteBlueprint, CurrentRevInfo, &Object, FOnMergeResolved());
+			Contents = GenerateMergeTabContents(Editor, BaseBlueprint, BaseRevInfo, RemoteBlueprint, CurrentRevInfo, &Object, FOnMergeResolved());
+		}
 	}
 
 	TSharedRef<SDockTab> Tab =  FGlobalTabmanager::Get()->InvokeTab(MergeToolTabId);
diff --git a/Engine/Source/Developer/Merge/Private/MergeUtils.cpp b/Engine/Source/Developer/Merge/Private/MergeUtils.cpp
new file mode 100644
--- /dev/null
+++ b/Engine/Source/Developer/Merge/Private/MergeUtils.cpp
@@ -0,0 +1,136 @@
+// Copyright 1998-2015 Epic Games, Inc. All Rights Reserved.
+
+#include "MergePrivatePCH.h"
+#include "MergeUtils.h"
+#include "ISourceControlModule.h"
+#include "SourceControlHelpers.h"
+#include "NotificationManager.h"
+
+#define LOCTEXT_NAMESPACE "MergeToolUtils"
+
+//------------------------------------------------------------------------------
+FSourceControlStatePtr FMergeToolUtils::GetSourceControlState(const FString& PackageName)
+{
+	ISourceControlProvider& SourceControlProvider = ISourceControlModule::Get().GetProvider();
+	TSharedRef<FUpdateStatus, ESPMode::ThreadSafe> UpdateStatusOperation = ISourceControlOperation::Create<FUpdateStatus>();
+	UpdateStatusOperation->SetUpdateHistory(true);
+	SourceControlProvider.Execute(UpdateStatusOperation, SourceControlHelpers::PackageFilename(PackageName));
+
+	FSourceControlStatePtr State = SourceControlProvider.GetState(SourceControlHelpers::PackageFilename(PackageName), EStateCacheUsage::Use);
+	if (!State.IsValid() || !State->IsSourceControlled() || !FPackageName::DoesPackageExist(PackageName))
+	{
+		return FSourceControlStatePtr();
+	}
+	else
+	{
+		return State;
+	}
+}
+
+//------------------------------------------------------------------------------
+UObject const* FMergeToolUtils::LoadRevision(const FString& AssetName, const ISourceControlRevision& DesiredRevision)
+{
+	FSlateNotificationManager& NotificationManager = FSlateNotificationManager::Get();
+
+	// Get the head revision of this package from source control
+	FString TempFileName;
+	if (DesiredRevision.Get(TempFileName))
+	{
+		// Forcibly disable compile on load in case we are loading old blueprints that might try to update/compile
+		TGuardValue<bool> DisableCompileOnLoad(GForceDisableBlueprintCompileOnLoad, true);
+
+		// Try and load that package
+		UPackage* TempPackage = LoadPackage(NULL, *TempFileName, LOAD_None);
+		if (TempPackage != NULL)
+		{
+			// Grab the old asset from that old package
+			UObject* OldObject = FindObject<UObject>(TempPackage, *AssetName);
+			if (OldObject)
+			{
+				return OldObject;
+			}
+			else
+			{
+				NotificationManager.AddNotification(
+					FText::Format(
+						LOCTEXT("MergedFailedToFindObject", "Aborted Load of {0} because we could not find an object named {1}" )
+						, FText::FromString(TempFileName)
+						, FText::FromString(AssetName) 
+					)
+				);
+			}
+		}
+		else
+		{
+			NotificationManager.AddNotification(
+				FText::Format(
+					LOCTEXT("MergedFailedToLoadPackage", "Aborted Load of {0} because we could not load the package")
+					, FText::FromString(TempFileName)
+				)
+			);
+		}
+	}
+	else
+	{
+		NotificationManager.AddNotification(
+			FText::Format(
+				LOCTEXT("MergedFailedToFindRevision", "Aborted Load of {0} because we could not get the requested revision")
+				, FText::FromString(TempFileName)
+			)
+		);
+	}
+	return NULL;
+}
+
+//------------------------------------------------------------------------------
+UObject const* FMergeToolUtils::LoadRevision(const FString& PackageName, const FRevisionInfo& DesiredRevision)
+{
+	const UObject* AssetRevision = nullptr;
+	if (UPackage* AssetPackage = LoadPackage(/*Outer =*/nullptr, *PackageName, LOAD_None))
+	{
+		TArray<UObject*> PackageObjs;
+		GetObjectsWithOuter(AssetPackage, PackageObjs, /*bIncludeNestedObjects =*/false);
+
+		for (UObject* PackageObj : PackageObjs)
+		{
+			if (PackageObj->IsAsset())
+			{
+				AssetRevision = LoadRevision(PackageObj, DesiredRevision);
+				break;
+			}
+		}
+	}
+	return AssetRevision;
+}
+
+//------------------------------------------------------------------------------
+UObject const* FMergeToolUtils::LoadRevision(const UObject* AssetObject, const FRevisionInfo& DesiredRevision)
+{
+	check(AssetObject->IsAsset());
+	
+	const UObject* AssetRevision = nullptr;
+	if (DesiredRevision.Revision < 0)
+	{
+		// an invalid revision number represents the local copy
+		AssetRevision = AssetObject;
+	}
+	else
+	{
+		FString const PackageName = AssetObject->GetOutermost()->GetName();
+
+		FSourceControlStatePtr SourceControlState = GetSourceControlState(PackageName);
+		if (SourceControlState.IsValid())
+		{
+			TSharedPtr<ISourceControlRevision, ESPMode::ThreadSafe> Revision = SourceControlState->FindHistoryRevision(DesiredRevision.Revision);
+			if (Revision.IsValid())
+			{
+				FString const AssetName = AssetObject->GetName();
+				AssetRevision = LoadRevision(AssetName, *Revision);
+			}
+		}
+	}
+
+	return AssetRevision;
+}
+
+#undef LOCTEXT_NAMESPACE
diff --git a/Engine/Source/Developer/Merge/Private/MergeUtils.h b/Engine/Source/Developer/Merge/Private/MergeUtils.h
new file mode 100644
--- /dev/null
+++ b/Engine/Source/Developer/Merge/Private/MergeUtils.h
@@ -0,0 +1,61 @@
+// Copyright 1998-2015 Epic Games, Inc. All Rights Reserved.
+
+#pragma once
+
+#include "ISourceControlState.h" // for FSourceControlStatePtr
+
+// Forward declarations
+class ISourceControlRevision;
+
+/**  */
+namespace EMergeAssetId
+{
+	enum Type
+	{
+		MergeRemote = 0,
+		MergeBase,
+		MergeLocal,
+		Count
+	};
+}
+
+/** Unified helper lib, serving shared utility functions across the merge module. */
+struct FMergeToolUtils
+{
+	/** 
+	 * Retrieves the source-control state of a specified package.
+	 *
+	 * @param  PackageName  The name of the asset package that you want state for.
+	 * @return An invalid source-control pointer if the package doesn't exist, 
+	 *         or if it doesn't have any source-control state; otherwise, a valid source-control state object.
+	 */
+	static FSourceControlStatePtr GetSourceControlState(const FString& PackageName);
+
+	/** 
+	 * Attempts to load the specified asset from the supplied source-control 
+	 * revision.
+	 *
+	 * @param  AssetName		The name of the asset you want loaded from the revision.
+	 * @param  DesiredRevision	The source-control revision you want loaded.
+	 * @return The asset object you wanted loaded, could be null if the asset wasn't found in the supplied revision.
+	 */
+	static UObject const* LoadRevision(const FString& AssetName, const ISourceControlRevision& DesiredRevision);
+
+	/** 
+	 * Attempts to load the supplied revision of the specified package.
+	 *
+	 * @param  PackageName		The name of the package you want loaded.
+	 * @param  DesiredRevision	Details the revision you want loaded.
+	 * @return The package's asset object, could be null if the specified revision could not be loaded.
+	 */
+	static UObject const* LoadRevision(const FString& PackageName, const FRevisionInfo& DesiredRevision);
+
+	/** 
+	 * Attempts to load the specified revision of the supplied asset.
+	 *
+	 * @param  AssetObject		The asset you want a different revision of.
+	 * @param  DesiredRevision	Details the revision you want loaded.
+	 * @return The requested version of the specified asset object, could be null if the supplied revision failed to load.
+	 */
+	static UObject const* LoadRevision(const UObject* AssetObject, const FRevisionInfo& DesiredRevision);
+};
\ No newline at end of file
diff --git a/Engine/Source/Developer/Merge/Private/SBlueprintMerge.cpp b/Engine/Source/Developer/Merge/Private/SBlueprintMerge.cpp
--- a/Engine/Source/Developer/Merge/Private/SBlueprintMerge.cpp
+++ b/Engine/Source/Developer/Merge/Private/SBlueprintMerge.cpp
@@ -13,6 +13,8 @@
 #include "SMergeTreeView.h"
 #include "SNotificationList.h"
 #include "NotificationManager.h"
+#include "SMergeAssetPickerView.h"
+#include "MergeUtils.h"
 
 #define LOCTEXT_NAMESPACE "SBlueprintMerge"
 
@@ -59,48 +61,25 @@ void SBlueprintMerge::Construct(const FArguments InArgs, const FBlueprintMergeDa
 {
 	check( InData.OwningEditor.Pin().IsValid() );
 
+	// reset state
 	Data = InData;
-	OnMergeResolved = InArgs._OnMergeResolved;
-	LocalBackupPath.Empty();
-
-	// we want to make a read only copy of the local blueprint, since we allow 
-	// the user to modify/mutate the merge result (which starts as the local copy)
-	if (InData.BlueprintLocal == GetTargetBlueprint())
+	bIsPickingAssets = InArgs._bForcePickAssets;
+	OnMergeResolved  = InArgs._OnMergeResolved;
+	
+	if (InData.BlueprintRemote != nullptr)
 	{
-		UBlueprint* LocalBlueprint = GetTargetBlueprint();
-		// we use this to suppress blueprint compilation during StaticDuplicateObject()
-		// @TODO: we don't remember exactly why this was needed, but I saw the 
-		//        the blueprint become unplacable and the file become unloadable
-		//        without this, but don't know exactly why
-		TGuardValue<bool> GuardDuplicatingReadOnly(LocalBlueprint->bDuplicatingReadOnly, true);
-
-		UPackage* TransientPackage = GetTransientPackage();
-		Data.BlueprintLocal = Cast<const UBlueprint>( StaticDuplicateObject(LocalBlueprint, TransientPackage,
-			*MakeUniqueObjectName(TransientPackage, LocalBlueprint->GetClass(), LocalBlueprint->GetFName()).ToString()) );
+		RemotePath = InData.BlueprintRemote->GetOutermost()->GetName();
+	}
+	if (InData.BlueprintBase != nullptr)
+	{
+		BasePath = InData.BlueprintBase->GetOutermost()->GetName();
+	}
+	if (InData.BlueprintLocal != nullptr)
+	{
+		LocalPath = InData.BlueprintLocal->GetOutermost()->GetName();
 	}
-
 	BackupSubDir = FPaths::GameSavedDir() / TEXT("Backup") / TEXT("Resolve_Backup[") + FDateTime::Now().ToString(TEXT("%Y-%m-%d-%H-%M-%S")) + TEXT("]");
 
-	// Because merge operations are so destructive and can be confusing, lets write backups of the files involved:
-	WriteBackup(*InData.BlueprintRemote->GetOutermost(), BackupSubDir, TEXT("RemoteAsset") + FPackageName::GetAssetPackageExtension());
-	WriteBackup(*InData.BlueprintBase->GetOutermost(), BackupSubDir, TEXT("CommonBaseAsset") + FPackageName::GetAssetPackageExtension());
-	LocalBackupPath = WriteBackup(*InData.BlueprintLocal->GetOutermost(), BackupSubDir, TEXT("LocalAsset") + FPackageName::GetAssetPackageExtension());
-
-	auto GraphView = SNew(SMergeGraphView, Data);
-	GraphControl.Widget = GraphView;
-	GraphControl.DiffControl = &GraphView.Get();
-	GraphControl.MergeControl = &GraphView.Get();
-
-	auto TreeView = SNew(SMergeTreeView, Data);
-	TreeControl.Widget = TreeView;
-	TreeControl.DiffControl = &TreeView.Get();
-	TreeControl.MergeControl = &TreeView.Get();
-
-	auto DetailsView = SNew(SMergeDetailsView, Data);
-	DetailsControl.Widget = DetailsView;
-	DetailsControl.DiffControl = &DetailsView.Get();
-	DetailsControl.MergeControl = &DetailsView.Get();
-
 	FToolBarBuilder ToolbarBuilder(TSharedPtr< FUICommandList >(), FMultiBoxCustomization::None);
 	ToolbarBuilder.AddToolBarButton( 
 		FUIAction( FExecuteAction::CreateSP(this, &SBlueprintMerge::PrevDiff), FCanExecuteAction::CreateSP(this, &SBlueprintMerge::HasPrevDiff) )
@@ -117,6 +96,11 @@ void SBlueprintMerge::Construct(const FArguments InArgs, const FBlueprintMergeDa
 		, FSlateIcon(FEditorStyle::GetStyleSetName(), "BlueprintMerge.NextDiff") 
 	);
 
+	auto IsInPassiveMode = [this]()
+	{
+		return !IsActivelyMerging();
+	};
+
 	// conflict navigation buttons:
 	ToolbarBuilder.AddSeparator();
 	ToolbarBuilder.AddToolBarButton( 
@@ -137,21 +121,41 @@ void SBlueprintMerge::Construct(const FArguments InArgs, const FBlueprintMergeDa
 	// buttons for finishing the merge:
 	ToolbarBuilder.AddSeparator(); 
 	ToolbarBuilder.AddToolBarButton(
-		FUIAction(FExecuteAction::CreateSP(this, &SBlueprintMerge::OnAcceptRemote))
+		FUIAction(
+			FExecuteAction::CreateSP(this, &SBlueprintMerge::OnStartMerge),
+			FCanExecuteAction::CreateSP(this, &SBlueprintMerge::CanStartMerge),
+			FIsActionChecked(),
+			FIsActionButtonVisible::CreateLambda(IsInPassiveMode))
+		, NAME_None
+		, LOCTEXT("StartMergeLabel", "Start Merge")
+		, LOCTEXT("StartMergeTooltip", "Loads the selected blueprints and switches to an active merge (using your selections for the remote/base/local)")
+		, FSlateIcon(FEditorStyle::GetStyleSetName(), "BlueprintMerge.StartMerge")
+	);
+	ToolbarBuilder.AddToolBarButton(
+		FUIAction(
+			FExecuteAction::CreateSP(this, &SBlueprintMerge::OnAcceptRemote),
+			FCanExecuteAction(), FIsActionChecked(),
+			FIsActionButtonVisible::CreateSP(this, &SBlueprintMerge::IsActivelyMerging))
 		, NAME_None
 		, LOCTEXT("AcceptRemoteLabel", "Accept Source")
 		, LOCTEXT("AcceptRemoteTooltip", "Complete the merge operation - Replaces the Blueprint with a copy of the remote file.")
 		, FSlateIcon(FEditorStyle::GetStyleSetName(), "BlueprintMerge.AcceptSource")
 	);
 	ToolbarBuilder.AddToolBarButton(
-		FUIAction(FExecuteAction::CreateSP(this, &SBlueprintMerge::OnAcceptLocal))
+		FUIAction(
+			FExecuteAction::CreateSP(this, &SBlueprintMerge::OnAcceptLocal),
+			FCanExecuteAction(), FIsActionChecked(),
+			FIsActionButtonVisible::CreateSP(this, &SBlueprintMerge::IsActivelyMerging))
 		, NAME_None
 		, LOCTEXT("AcceptLocalLabel", "Accept Target")
 		, LOCTEXT("AcceptLocalTooltip", "Complete the merge operation - Leaves the target Blueprint unchanged.")
 		, FSlateIcon(FEditorStyle::GetStyleSetName(), "BlueprintMerge.AcceptTarget")
 	);
 	ToolbarBuilder.AddToolBarButton(
-		FUIAction(FExecuteAction::CreateSP(this, &SBlueprintMerge::OnFinishMerge))
+		FUIAction(
+			FExecuteAction::CreateSP(this, &SBlueprintMerge::OnFinishMerge), 
+			FCanExecuteAction(), FIsActionChecked(),
+			FIsActionButtonVisible::CreateSP(this, &SBlueprintMerge::IsActivelyMerging))
 		, NAME_None
 		, LOCTEXT("FinishMergeLabel", "Finish Merge")
 		, LOCTEXT("FinishMergeTooltip", "Complete the merge operation - saves the blueprint and resolves the conflict with the SCC provider")
@@ -180,10 +184,24 @@ void SBlueprintMerge::Construct(const FArguments InArgs, const FBlueprintMergeDa
 		+ SVerticalBox::Slot()
 		[
 			SAssignNew(MainView, SBorder)
-			.BorderImage(FEditorStyle::GetBrush("NoBorder"))
+				.BorderImage(FEditorStyle::GetBrush("NoBorder"))
 		]
 	];
 
+	if (IsActivelyMerging())
+	{
+		OnStartMerge();
+	}
+	else
+	{
+		bIsPickingAssets = true;
+
+		auto AssetPickerView = SNew(SMergeAssetPickerView, InData).OnAssetChanged(this, &SBlueprintMerge::OnMergeAssetSelected);
+		AssetPickerControl.Widget = AssetPickerView;
+		AssetPickerControl.DiffControl  = &AssetPickerView.Get();
+		AssetPickerControl.MergeControl = &AssetPickerView.Get();
+	}
+
 	Data.OwningEditor.Pin()->OnModeSet().AddSP( this, &SBlueprintMerge::OnModeChanged );
 	OnModeChanged( Data.OwningEditor.Pin()->GetCurrentMode() );
 }
@@ -233,6 +251,82 @@ bool SBlueprintMerge::HasPrevConflict() const
 	return CurrentMergeControl && CurrentMergeControl->HasPrevConflict();
 }
 
+void SBlueprintMerge::OnStartMerge()
+{
+	if (Data.BlueprintRemote == nullptr)
+	{
+		Data.BlueprintRemote = Cast<UBlueprint>(FMergeToolUtils::LoadRevision(RemotePath, Data.RevisionRemote));
+	}
+	if (Data.BlueprintBase == nullptr)
+	{
+		Data.BlueprintBase = Cast<UBlueprint>(FMergeToolUtils::LoadRevision(BasePath, Data.RevisionBase));
+	}
+	if (Data.BlueprintLocal == nullptr)
+	{
+		Data.BlueprintLocal = Cast<UBlueprint>(FMergeToolUtils::LoadRevision(LocalPath, Data.RevisionLocal));
+	}
+	LocalBackupPath.Empty();
+
+	if ((Data.BlueprintRemote != nullptr) && (Data.BlueprintBase != nullptr) && (Data.BlueprintLocal != nullptr))
+	{
+		// Because merge operations are so destructive and can be confusing, lets write backups of the files involved:
+		WriteBackup(*Data.BlueprintRemote->GetOutermost(), BackupSubDir, TEXT("RemoteAsset") + FPackageName::GetAssetPackageExtension());
+		WriteBackup(*Data.BlueprintBase->GetOutermost(), BackupSubDir, TEXT("CommonBaseAsset") + FPackageName::GetAssetPackageExtension());
+		LocalBackupPath = WriteBackup(*Data.BlueprintLocal->GetOutermost(), BackupSubDir, TEXT("LocalAsset") + FPackageName::GetAssetPackageExtension());
+
+		// we want to make a read only copy of the local blueprint, since we allow 
+		// the user to modify/mutate the merge result (which starts as the local copy)
+		if (Data.BlueprintLocal == GetTargetBlueprint())
+		{
+			UBlueprint* LocalBlueprint = GetTargetBlueprint();
+			// we use this to suppress blueprint compilation during StaticDuplicateObject()
+			// @TODO: we don't remember exactly why this was needed, but I saw the 
+			//        the blueprint become unplacable and the file become unloadable
+			//        without this, but don't know exactly why
+			TGuardValue<bool> GuardDuplicatingReadOnly(LocalBlueprint->bDuplicatingReadOnly, true);
+
+			UPackage* TransientPackage = GetTransientPackage();
+			Data.BlueprintLocal = Cast<const UBlueprint>(StaticDuplicateObject(LocalBlueprint, TransientPackage,
+				*MakeUniqueObjectName(TransientPackage, LocalBlueprint->GetClass(), LocalBlueprint->GetFName()).ToString()));
+		}
+		else
+		{
+			// since we didn't need to make a copy of the local-blueprint, then
+			// we're most likely using a different asset for "local"; in that 
+			// case, we don't want to use a copy of the file in OnAcceptLocal()
+			// (the name of the blueprint would be all off, etc.)... instead,
+			// we want to just duplicate the Blueprint pointer (like we do for
+			// OnAcceptRemote), so we leverage the emptiness of LocalBackupPath
+			// as a sentinel value (to determine which method we choose)
+			LocalBackupPath.Empty();
+		}
+
+		auto GraphView = SNew(SMergeGraphView, Data);
+		GraphControl.Widget = GraphView;
+		GraphControl.DiffControl = &GraphView.Get();
+		GraphControl.MergeControl = &GraphView.Get();
+
+		auto TreeView = SNew(SMergeTreeView, Data);
+		TreeControl.Widget = TreeView;
+		TreeControl.DiffControl = &TreeView.Get();
+		TreeControl.MergeControl = &TreeView.Get();
+
+		auto DetailsView = SNew(SMergeDetailsView, Data);
+		DetailsControl.Widget = DetailsView;
+		DetailsControl.DiffControl = &DetailsView.Get();
+		DetailsControl.MergeControl = &DetailsView.Get();
+	
+		bIsPickingAssets = false;
+		OnModeChanged(Data.OwningEditor.Pin()->GetCurrentMode());
+	}
+	else
+	{
+		const FText ErrorMessage = LOCTEXT("FailedMergeLoad", "Failed to load asset(s) for merge.");
+		FSlateNotificationManager::Get().AddNotification(FNotificationInfo(ErrorMessage));
+	}
+
+}
+
 void SBlueprintMerge::OnFinishMerge()
 {
 	ResolveMerge(GetTargetBlueprint());
@@ -254,7 +348,13 @@ void SBlueprintMerge::OnCancelClicked()
 
 void SBlueprintMerge::OnModeChanged(FName NewMode)
 {
-	if (NewMode == FBlueprintEditorApplicationModes::StandardBlueprintEditorMode ||
+	if (!IsActivelyMerging())
+	{
+		MainView->SetContent(AssetPickerControl.Widget.ToSharedRef());
+		CurrentDiffControl = AssetPickerControl.DiffControl;
+		CurrentMergeControl = AssetPickerControl.MergeControl;
+	}
+	else if (NewMode == FBlueprintEditorApplicationModes::StandardBlueprintEditorMode ||
 		NewMode == FBlueprintEditorApplicationModes::BlueprintMacroMode)
 	{
 		MainView->SetContent( GraphControl.Widget.ToSharedRef() );
@@ -295,38 +395,53 @@ void SBlueprintMerge::OnAcceptRemote()
 
 void SBlueprintMerge::OnAcceptLocal()
 {
-	UBlueprint* TargetBlueprint = GetTargetBlueprint();
-	UPackage* TargetPackage = TargetBlueprint->GetOutermost();
+	UBlueprint* NewBlueprint = nullptr;
 
-	FString PackageFilename;
-	if (FPackageName::DoesPackageExist(TargetPackage->GetName(), /*Guid =*/nullptr, &PackageFilename))
+	UBlueprint* TargetBlueprint = GetTargetBlueprint();
+	if (PromptUserIfUndoBufferToBeCleared(TargetBlueprint) == EAppReturnType::Yes)
 	{
-		if (PromptUserIfUndoBufferToBeCleared(TargetBlueprint) == EAppReturnType::Yes)
+		// we use the emptiness of LocalBackupPath as a sentinel value (to 
+		// determine how we should replace the TargetBlueprint)... if 
+		// Data.BlueprintLocal is malformed (and not ok to copy), then 
+		// LocalBackupPath should be valid (and we need to copy the file on disk 
+		// instead)
+		if (LocalBackupPath.IsEmpty())
 		{
-			FString SrcFilename = LocalBackupPath;
-			auto OverwriteBlueprintFile = [TargetBlueprint, &PackageFilename, &SrcFilename](UBlueprint* /*UnloadedBP*/)
+			NewBlueprint = FKismetEditorUtilities::ReplaceBlueprint(TargetBlueprint, Data.BlueprintLocal);
+		}
+		else
+		{
+			UPackage* TargetPackage = TargetBlueprint->GetOutermost();
+
+			FString PackageFilename;
+			if (FPackageName::DoesPackageExist(TargetPackage->GetName(), /*Guid =*/nullptr, &PackageFilename))
 			{
-				if (IFileManager::Get().Copy(*PackageFilename, *SrcFilename) != COPY_OK)
+				FString SrcFilename = LocalBackupPath;
+				auto OverwriteBlueprintFile = [TargetBlueprint, &PackageFilename, &SrcFilename](UBlueprint* /*UnloadedBP*/)
 				{
-					const FText TargetName = FText::FromName(TargetBlueprint->GetFName());
-					const FText ErrorMessage = FText::Format(LOCTEXT("FailedMergeLocalCopy", "Failed to overwrite {0} target file."), TargetName);
-					FSlateNotificationManager::Get().AddNotification(FNotificationInfo(ErrorMessage));
-				}
-			};
-			// we have to wait until the package is unloaded (and the file is 
-			// unlocked) before we can overwrite the file; so we use defer the
-			// copy until then by adding it to the OnBlueprintUnloaded delegate
-			FKismetEditorUtilities::FOnBlueprintUnloaded::FDelegate OnPackageUnloaded = FKismetEditorUtilities::FOnBlueprintUnloaded::FDelegate::CreateLambda(OverwriteBlueprintFile);
-			FKismetEditorUtilities::OnBlueprintUnloaded.Add(OnPackageUnloaded);
-
-
-			if (UBlueprint* NewBlueprint = FKismetEditorUtilities::ReloadBlueprint(TargetBlueprint))
-			{
-				ResolveMerge(NewBlueprint);
+					if (IFileManager::Get().Copy(*PackageFilename, *SrcFilename) != COPY_OK)
+					{
+						const FText TargetName = FText::FromName(TargetBlueprint->GetFName());
+						const FText ErrorMessage = FText::Format(LOCTEXT("FailedMergeLocalCopy", "Failed to overwrite {0} target file."), TargetName);
+						FSlateNotificationManager::Get().AddNotification(FNotificationInfo(ErrorMessage));
+					}
+				};
+				// we have to wait until the package is unloaded (and the file is 
+				// unlocked) before we can overwrite the file; so we defer the
+				// copy until then by adding it to the OnBlueprintUnloaded delegate
+				FKismetEditorUtilities::FOnBlueprintUnloaded::FDelegate OnPackageUnloaded = FKismetEditorUtilities::FOnBlueprintUnloaded::FDelegate::CreateLambda(OverwriteBlueprintFile);
+					
+				FKismetEditorUtilities::OnBlueprintUnloaded.Add(OnPackageUnloaded);
+				NewBlueprint = FKismetEditorUtilities::ReloadBlueprint(TargetBlueprint);
+				FKismetEditorUtilities::OnBlueprintUnloaded.Remove(OnPackageUnloaded);
 			}
-			FKismetEditorUtilities::OnBlueprintUnloaded.Remove(OnPackageUnloaded);
 		}
 	}
+
+	if (NewBlueprint != nullptr)
+	{
+		ResolveMerge(NewBlueprint);
+	}
 }
 
 void SBlueprintMerge::ResolveMerge(UBlueprint* ResultantBlueprint)
@@ -367,4 +482,43 @@ void SBlueprintMerge::ResolveMerge(UBlueprint* ResultantBlueprint)
 	}
 }
 
+bool SBlueprintMerge::IsActivelyMerging() const
+{
+	return !bIsPickingAssets && (Data.BlueprintRemote != nullptr) &&
+		(Data.BlueprintBase != nullptr) && (Data.BlueprintLocal != nullptr);
+}
+
+bool SBlueprintMerge::CanStartMerge() const
+{
+	return !IsActivelyMerging() && !RemotePath.IsEmpty() && !BasePath.IsEmpty() && !LocalPath.IsEmpty();
+}
+
+void SBlueprintMerge::OnMergeAssetSelected(EMergeAssetId::Type AssetId, const FAssetRevisionInfo& AssetInfo)
+{
+	switch (AssetId)
+	{
+	case EMergeAssetId::MergeRemote:
+		{
+			RemotePath = AssetInfo.AssetName;
+			Data.RevisionRemote  = AssetInfo.Revision;
+			Data.BlueprintRemote = nullptr;
+			break;
+		}
+	case EMergeAssetId::MergeBase:
+		{
+			BasePath = AssetInfo.AssetName;
+			Data.RevisionBase  = AssetInfo.Revision;
+			Data.BlueprintBase = nullptr;
+			break;
+		}
+	case EMergeAssetId::MergeLocal:
+		{
+			LocalPath = AssetInfo.AssetName;
+			Data.RevisionLocal  = AssetInfo.Revision;
+			Data.BlueprintLocal = nullptr;
+			break;
+		}
+	}
+}
+
 #undef LOCTEXT_NAMESPACE
diff --git a/Engine/Source/Developer/Merge/Private/SBlueprintMerge.h b/Engine/Source/Developer/Merge/Private/SBlueprintMerge.h
--- a/Engine/Source/Developer/Merge/Private/SBlueprintMerge.h
+++ b/Engine/Source/Developer/Merge/Private/SBlueprintMerge.h
@@ -2,8 +2,10 @@
 #pragma once
 
 #include "BlueprintMergeData.h"
+#include "MergeUtils.h" // for EMergeAssetId
 
-class FBlueprintEditor;
+class  FBlueprintEditor;
+struct FAssetRevisionInfo;
 
 class MERGE_API SBlueprintMerge : public SCompoundWidget
 {
@@ -11,7 +13,9 @@ public:
 	SBlueprintMerge();
 
 	SLATE_BEGIN_ARGS(SBlueprintMerge)
+		: _bForcePickAssets(false)
 	{}
+		SLATE_ARGUMENT(bool, bForcePickAssets)
 		SLATE_EVENT(FOnMergeResolved, OnMergeResolved)
 	SLATE_END_ARGS()
 
@@ -32,6 +36,7 @@ private:
 	bool HasNextConflict() const;
 	bool HasPrevConflict() const;
 
+	void OnStartMerge();
 	void OnFinishMerge();
 	void OnCancelClicked();
 	void OnModeChanged(FName NewMode);
@@ -39,12 +44,56 @@ private:
 	void OnAcceptLocal();
 	void ResolveMerge(UBlueprint* Result);
 
+	/** 
+	 * If the user has yet to pick their remote/base/local assets, then they're
+	 * not "actively" merging yet. This query checks the state of the merge-tool
+	 * to see if the user has selected to "Start Merge" yet.
+	 *
+	 * @return False if the AssetPickerControl should be up, otherwise true.
+	 */
+	bool IsActivelyMerging() const;
+
+	/** 
+	 * The user cannot start a merge until he/she has chosen a remote, base, &
+	 * local asset/revision to use in the merge. This checks the state of their 
+	 * choices.
+	 *
+	 * @return True, if the user has selected all three merge assets.
+	 */
+	bool CanStartMerge() const;
+
+	/** 
+	 * Callback function, utilized every time the user picks an asset/revision  
+	 * in the AssetPickerControl view. Records the user's choices, so this can 
+	 * start the merge appropriately.
+	 * 
+	 * @param  AssetId		Defines the asset that was altered (remote, base, or local).
+	 * @param  AssetInfo	Defines the asset selection (name, revision, etc.)
+	 */
+	void OnMergeAssetSelected(EMergeAssetId::Type AssetId, const FAssetRevisionInfo& AssetInfo);
+
 	FBlueprintMergeData		Data; 
 	FString					BackupSubDir;
-	FString					LocalBackupPath;
 
 	TSharedPtr<SBorder>		MainView;
 
+	/** 
+	 * We track the package name-paths for the remote, base, & local assets (so
+	 * we know what to load when the user starts an active merge)... Used to 
+	 * determine when a merge can be started.
+	 */
+	FString	RemotePath;
+	FString	BasePath;
+	FString	LocalPath;
+	/** 
+	 * When we make a malformed copy (readable data only) of the local 
+	 * Blueprint, then we fill this in with a backup file path; that way, the 
+	 * file can be copied for OnAcceptLocal(), instead of the loaded Blueprint 
+	 * object (which is "malformed").
+	 */
+	FString	LocalBackupPath;
+
+
 	struct FMergeControl
 	{
 		FMergeControl()
@@ -62,7 +111,9 @@ private:
 	FMergeControl GraphControl;
 	FMergeControl TreeControl;
 	FMergeControl DetailsControl;
+	FMergeControl AssetPickerControl;
 	
+	bool bIsPickingAssets;
 	FOnMergeResolved OnMergeResolved;
 
 	class IDiffControl* CurrentDiffControl;
diff --git a/Engine/Source/Developer/Merge/Private/SMergeAssetPickerView.cpp b/Engine/Source/Developer/Merge/Private/SMergeAssetPickerView.cpp
new file mode 100644
--- /dev/null
+++ b/Engine/Source/Developer/Merge/Private/SMergeAssetPickerView.cpp
@@ -0,0 +1,435 @@
+// Copyright 1998-2015 Epic Games, Inc. All Rights Reserved.
+
+#include "MergePrivatePCH.h"
+#include "SMergeAssetPickerView.h"
+#include "Editor/ContentBrowser/Public/ContentBrowserModule.h"
+#include "SBlueprintRevisionMenu.h"
+#include "Editor.h"
+
+#define LOCTEXT_NAMESPACE "SMergeAssetPickerView"
+
+/*******************************************************************************
+ * SMergeAssetPickerPanel
+ ******************************************************************************/
+
+class SMergeAssetPickerPanel : public SCompoundWidget
+{
+	DECLARE_DELEGATE_OneParam(FOnAssetChanged, const UObject*);
+	DECLARE_DELEGATE_OneParam(FOnRevisionChanged, const FRevisionInfo&);
+
+public:
+	SLATE_BEGIN_ARGS(SMergeAssetPickerPanel)
+		: _SelectedAsset(nullptr)
+		, _RevisionInfo(FRevisionInfo::InvalidRevision())
+		, _AssetClass(UBlueprint::StaticClass())
+	{}
+		SLATE_ARGUMENT(TSubclassOf<UBlueprint>, AssetClass)
+		SLATE_ARGUMENT(const UBlueprint*, SelectedAsset)
+		SLATE_ARGUMENT(FRevisionInfo, RevisionInfo)
+
+		SLATE_EVENT(FOnAssetChanged, OnAssetChanged)
+		SLATE_EVENT(FOnRevisionChanged, OnRevisionChanged)
+	SLATE_END_ARGS()
+
+	/**  */
+	SMergeAssetPickerPanel() 
+		: SelectedAsset(nullptr)
+		, SelectedRevision(FRevisionInfo::InvalidRevision())
+	{}
+
+	/** 
+	 * Builds the slate UI for SMergeAssetPickerPanel (as well as caches off any
+	 * important data that was supplied as part of the FArguments struct).
+	 *
+	 * @param  Args	 The set of slate arguments that are used to customize this panel.
+	 */
+	void Construct(const FArguments Args);
+
+	/** 
+	 * Prior to launch, sometimes temp assets are created for the merge-tool 
+	 * (old revisions are dumped into the temp directory, etc.). This function
+	 * determines if the currently selected asset (for this panel) is one of 
+	 * those temporary asset files.
+	 *
+	 * @return True if the selected asset is a temporary one, other wise false (if it is a legit asset, or if one isn't selected).
+	 */
+	bool IsTempAssetSelected() const;
+
+private:
+	TSharedRef<SWidget> MakeAssetPicker();
+	FText  GetAssetPickerTextValue() const;
+	FReply OnUseSelectedAssetClick();
+	void   ResetSelectedAsset(const UObject* NewSelection);
+	bool   IsBrowseButtonEnabled() const;
+	FReply OnBrowseToPickedAsset();
+
+	TSharedRef<SWidget> MakeAssetPickerMenu();
+	void OnAssetSelected(FAssetData const& AssetData);
+
+	TSharedRef<SWidget> MakeRevisionPicker();
+	FText GetRevisionTextValue() const;
+
+	TSharedRef<SWidget> MakeRevisionPickerMenu();
+	void OnRevisionSelected(FRevisionInfo const& RevisionInfo);
+
+	TSubclassOf<UBlueprint> AssetClass;
+	UBlueprint const* SelectedAsset;
+	FOnAssetChanged OnAssetChanged;
+	FRevisionInfo SelectedRevision;
+	FOnRevisionChanged OnRevisionChanged;
+	TSharedPtr<SMenuAnchor> AssetPicker;
+};
+
+//------------------------------------------------------------------------------
+void SMergeAssetPickerPanel::Construct(const FArguments InArgs)
+{
+	AssetClass        = InArgs._AssetClass;
+	SelectedAsset     = InArgs._SelectedAsset;
+	SelectedRevision  = InArgs._RevisionInfo;
+	OnAssetChanged    = InArgs._OnAssetChanged;
+	OnRevisionChanged = InArgs._OnRevisionChanged;
+
+	ChildSlot
+	[
+		SNew(SBorder)
+			.BorderImage(FEditorStyle::GetBrush("ToolPanel.GroupBorder"))
+		[
+			SNew(SBox)
+				.HAlign(HAlign_Center)
+				.VAlign(VAlign_Center)
+			[
+				SNew(SVerticalBox)
+				+ SVerticalBox::Slot()
+					.HAlign(HAlign_Left)
+					.AutoHeight()
+					.Padding(0, 2)
+				[
+					MakeAssetPicker()
+				]
+				+ SVerticalBox::Slot()
+					.HAlign(HAlign_Left)
+					.AutoHeight()
+					.Padding(0, 2)
+				[
+					MakeRevisionPicker()
+				]
+			]
+		]
+	];
+}
+
+//------------------------------------------------------------------------------
+bool SMergeAssetPickerPanel::IsTempAssetSelected() const
+{
+	return (SelectedAsset != nullptr) && (!SelectedAsset->IsAsset() ||(SelectedAsset->GetOutermost() == GetTransientPackage()) || 
+		SelectedAsset->GetPathName().StartsWith("/Temp/"));
+}
+
+//------------------------------------------------------------------------------
+TSharedRef<SWidget> SMergeAssetPickerPanel::MakeAssetPicker()
+{
+	TSharedPtr<SButton> BrowseButton;
+	SAssignNew(BrowseButton, SButton)
+		.ButtonStyle(FEditorStyle::Get(), "NoBorder")
+		.OnClicked(FOnClicked::CreateSP(this, &SMergeAssetPickerPanel::OnBrowseToPickedAsset))
+		.ContentPadding(0)
+		.ToolTipText(NSLOCTEXT("GraphEditor", "ObjectGraphPin_Browse", "Show the selected asset in the content browser.").ToString())
+	[
+		SNew(SImage)
+			.Image(FEditorStyle::GetBrush(TEXT("PropertyWindow.Button_Browse")))
+	];
+	BrowseButton->SetEnabled(TAttribute<bool>(this, &SMergeAssetPickerPanel::IsBrowseButtonEnabled));
+
+	return SNew(SHorizontalBox)
+	+ SHorizontalBox::Slot()
+		.Padding(2, 0)
+		.MaxWidth(150.0f)
+	[
+		SAssignNew(AssetPicker, SComboButton)
+			.ButtonStyle(FEditorStyle::Get(), "PropertyEditor.AssetComboStyle")
+			.ContentPadding(FMargin(2, 2, 2, 1))
+			.OnGetMenuContent(this, &SMergeAssetPickerPanel::MakeAssetPickerMenu)
+			.ButtonContent()
+		[
+			SNew(SBox).WidthOverride(150.f)
+			[
+				SNew(STextBlock)
+					.TextStyle(FEditorStyle::Get(), "PropertyEditor.AssetClass")
+					.Font(FEditorStyle::GetFontStyle("PropertyWindow.NormalFont"))
+					.Text(this, &SMergeAssetPickerPanel::GetAssetPickerTextValue)
+			]
+		]
+	]
+	+ SHorizontalBox::Slot()
+		.AutoWidth()
+		.Padding(1, 0)
+		.VAlign(VAlign_Center)
+	[
+		SNew(SButton)
+			.ButtonStyle(FEditorStyle::Get(), "NoBorder")
+			.OnClicked(FOnClicked::CreateSP(this, &SMergeAssetPickerPanel::OnUseSelectedAssetClick))
+			.ContentPadding(1.f)
+			.ToolTipText(NSLOCTEXT("GraphEditor", "ObjectGraphPin_Use", "Use content browser selection.").ToString())
+		[
+			SNew(SImage)
+				.Image(FEditorStyle::GetBrush(TEXT("PropertyWindow.Button_Use")))
+		]
+	]
+	+ SHorizontalBox::Slot()
+		.AutoWidth()
+		.Padding(1, 0)
+		.VAlign(VAlign_Center)
+	[
+		BrowseButton.ToSharedRef()
+	];
+}
+
+//------------------------------------------------------------------------------
+FText SMergeAssetPickerPanel::GetAssetPickerTextValue() const
+{
+	FText AssetPickerTitle = LOCTEXT("SelectAsset", "Select an Asset...");
+	if (SelectedAsset != nullptr)
+	{
+		if (UClass* BlueprintClass = SelectedAsset->GeneratedClass)
+		{
+			AssetPickerTitle = BlueprintClass->GetDisplayNameText();
+		}
+		else
+		{
+			AssetPickerTitle = FText::FromString(SelectedAsset->GetName());
+		}
+	}
+	return AssetPickerTitle;
+}
+
+//------------------------------------------------------------------------------
+FReply SMergeAssetPickerPanel::OnUseSelectedAssetClick()
+{
+	FEditorDelegates::LoadSelectedAssetsIfNeeded.Broadcast();
+
+	UObject* SelectedObject = GEditor->GetSelectedObjects()->GetTop(AssetClass);
+	if (SelectedObject != nullptr)
+	{
+		ResetSelectedAsset(SelectedObject);
+	}
+
+	return FReply::Handled();
+}
+
+//------------------------------------------------------------------------------
+void SMergeAssetPickerPanel::ResetSelectedAsset(const UObject* NewSelection)
+{
+	SelectedAsset = Cast<UBlueprint>(NewSelection);
+	OnAssetChanged.ExecuteIfBound(SelectedAsset);
+
+	// reset the revision (to current)
+	SelectedRevision = FRevisionInfo::InvalidRevision();
+	OnRevisionChanged.ExecuteIfBound(SelectedRevision);
+}
+
+//------------------------------------------------------------------------------
+bool SMergeAssetPickerPanel::IsBrowseButtonEnabled() const
+{
+	return (SelectedAsset != nullptr);
+}
+
+//------------------------------------------------------------------------------
+FReply SMergeAssetPickerPanel::OnBrowseToPickedAsset()
+{
+	if (SelectedAsset != nullptr)
+	{
+		TArray<UObject*> BrowseObjects;
+		// @TODO: SyncBrowserToObjects() should really take an array of const objects
+		BrowseObjects.Add(const_cast<UBlueprint*>(SelectedAsset));
+
+		GEditor->SyncBrowserToObjects(BrowseObjects);
+	}
+	return FReply::Handled();
+}
+
+//------------------------------------------------------------------------------
+TSharedRef<SWidget> SMergeAssetPickerPanel::MakeAssetPickerMenu()
+{
+	FAssetPickerConfig AssetPickerConfig;
+	AssetPickerConfig.Filter.ClassNames.Add(AssetClass->GetFName());
+	AssetPickerConfig.bAllowNullSelection = true;
+	AssetPickerConfig.Filter.bRecursiveClasses = true;
+	AssetPickerConfig.OnAssetSelected = FOnAssetSelected::CreateSP(this, &SMergeAssetPickerPanel::OnAssetSelected);
+	AssetPickerConfig.ThumbnailScale = 0;
+	AssetPickerConfig.InitialAssetViewType = EAssetViewType::List;
+	AssetPickerConfig.bAllowDragging = false;
+
+	FContentBrowserModule& ContentBrowserModule = FModuleManager::Get().LoadModuleChecked<FContentBrowserModule>(TEXT("ContentBrowser"));
+
+	return SNew(SBox)
+		.WidthOverride(300)
+		.HeightOverride(300)
+	[
+		SNew(SBorder)
+			.BorderImage(FEditorStyle::GetBrush("Menu.Background"))
+		[
+			ContentBrowserModule.Get().CreateAssetPicker(AssetPickerConfig)
+		]
+	];
+}
+
+//------------------------------------------------------------------------------
+void SMergeAssetPickerPanel::OnAssetSelected(FAssetData const& AssetData)
+{
+	AssetPicker->SetIsOpen(false);
+	ResetSelectedAsset(AssetData.GetAsset());
+}
+
+//------------------------------------------------------------------------------
+TSharedRef<SWidget> SMergeAssetPickerPanel::MakeRevisionPicker()
+{
+	return SNew(SHorizontalBox)
+	+ SHorizontalBox::Slot()
+		.Padding(2, 0)
+		.MaxWidth(150.0f)
+	[
+		SAssignNew(AssetPicker, SComboButton)
+			.ButtonStyle(FEditorStyle::Get(), "PropertyEditor.AssetComboStyle")
+			.ContentPadding(FMargin(2, 2, 2, 1))
+			.OnGetMenuContent(this, &SMergeAssetPickerPanel::MakeRevisionPickerMenu)
+			.ButtonContent()
+		[
+			SNew(SBox).WidthOverride(150.f)
+			[
+				SNew(STextBlock)
+					.TextStyle(FEditorStyle::Get(), "PropertyEditor.AssetClass")
+					.Font(FEditorStyle::GetFontStyle("PropertyWindow.NormalFont"))
+					.Text(this, &SMergeAssetPickerPanel::GetRevisionTextValue)
+			]
+		]
+	];
+}
+
+//------------------------------------------------------------------------------
+FText SMergeAssetPickerPanel::GetRevisionTextValue() const
+{
+	FText RevisionText = LOCTEXT("UnknownRevision", "Unknown Revision");
+	if (SelectedAsset == nullptr)
+	{
+		RevisionText = LOCTEXT("PickRevision", "Pick a Revision...");
+	}
+	else if (SelectedRevision.Revision > -1)
+	{
+		RevisionText = FText::Format(LOCTEXT("RevisionNum", "Revision {0}"), FText::AsNumber(SelectedRevision.Revision));
+	}
+	else if (!IsTempAssetSelected())
+	{
+		RevisionText = LOCTEXT("LocalRevision", "Local Copy");
+	}
+	else
+	{
+		RevisionText = LOCTEXT("UnknownRevision", "Unknown Revision");
+	}
+	return RevisionText;
+}
+
+//------------------------------------------------------------------------------
+TSharedRef<SWidget> SMergeAssetPickerPanel::MakeRevisionPickerMenu()
+{
+	if (SelectedAsset == nullptr)
+	{
+		return SNew(SBorder)
+			.BorderImage(FEditorStyle::GetBrush("Menu.Background"))
+		[
+			SNew(STextBlock).Text(LOCTEXT("SelectAssetFirst", "Select an asset first."))
+		];
+	}
+	else if (!IsTempAssetSelected())
+	{
+		return SNew(SBlueprintRevisionMenu, Cast<UBlueprint>(SelectedAsset))
+			.bIncludeLocalRevision(true)
+			.OnRevisionSelected(this, &SMergeAssetPickerPanel::OnRevisionSelected);
+	}
+	else
+	{
+		FMenuBuilder MenuBuilder(/*bInShouldCloseWindowAfterMenuSelection =*/true, /*CommandList =*/nullptr);
+		MenuBuilder.BeginSection("RevisionSection", LOCTEXT("Revisions", "Revisions"));
+		{
+			FText const RevisionTooltip = LOCTEXT("UnknownRevisionTip",
+				"The selected asset must be a temporary one created by the merge-tool (with this preordained revision).");
+			MenuBuilder.AddMenuEntry(TAttribute<FText>(this, &SMergeAssetPickerPanel::GetRevisionTextValue),
+				RevisionTooltip, FSlateIcon(), FUIAction());
+		}
+		MenuBuilder.EndSection();
+
+		return SNew(SBorder)
+			.BorderImage(FEditorStyle::GetBrush("Menu.Background"))
+		[
+			MenuBuilder.MakeWidget()
+		];
+	}
+}
+
+//------------------------------------------------------------------------------
+void SMergeAssetPickerPanel::OnRevisionSelected(FRevisionInfo const& RevisionInfo)
+{
+	SelectedRevision = RevisionInfo;
+	OnRevisionChanged.ExecuteIfBound(RevisionInfo);
+}
+
+/*******************************************************************************
+ * SMergeAssetPickerView
+ ******************************************************************************/
+
+//------------------------------------------------------------------------------
+void SMergeAssetPickerView::Construct(const FArguments InArgs, const FBlueprintMergeData& InData)
+{
+	OnAssetChanged = InArgs._OnAssetChanged;
+
+	ChildSlot
+	[
+		SNew(SSplitter)
+		+ SSplitter::Slot()
+		[
+			SNew(SMergeAssetPickerPanel)
+				.SelectedAsset(InData.BlueprintRemote)
+				.RevisionInfo(InData.RevisionRemote)
+				.OnRevisionChanged(this, &SMergeAssetPickerView::HandleRevisionChange, EMergeAssetId::MergeRemote)
+				.OnAssetChanged(this, &SMergeAssetPickerView::HandleAssetChange, EMergeAssetId::MergeRemote)
+		] 
+		+ SSplitter::Slot()
+		[
+			SNew(SMergeAssetPickerPanel)
+				.SelectedAsset(InData.BlueprintBase)
+				.RevisionInfo(InData.RevisionBase)
+				.OnRevisionChanged(this, &SMergeAssetPickerView::HandleRevisionChange, EMergeAssetId::MergeBase)
+				.OnAssetChanged(this, &SMergeAssetPickerView::HandleAssetChange, EMergeAssetId::MergeBase)
+		]
+		+ SSplitter::Slot()
+		[
+			SNew(SMergeAssetPickerPanel)
+				.SelectedAsset(InData.BlueprintLocal)
+				.RevisionInfo(InData.RevisionLocal)
+				.OnRevisionChanged(this, &SMergeAssetPickerView::HandleRevisionChange, EMergeAssetId::MergeLocal)
+				.OnAssetChanged(this, &SMergeAssetPickerView::HandleAssetChange, EMergeAssetId::MergeLocal)
+		]
+	];
+}
+
+//------------------------------------------------------------------------------
+void SMergeAssetPickerView::HandleAssetChange(const UObject* NewAsset, EMergeAssetId::Type PanelId)
+{
+	if (NewAsset == nullptr)
+	{
+		MergeAssetSet[PanelId].AssetName.Empty();
+	}
+	else
+	{
+		MergeAssetSet[PanelId].AssetName = NewAsset->GetOutermost()->GetName();
+	}
+	OnAssetChanged.ExecuteIfBound(PanelId, MergeAssetSet[PanelId]);
+}
+
+//------------------------------------------------------------------------------
+void SMergeAssetPickerView::HandleRevisionChange(const FRevisionInfo& NewRevision, EMergeAssetId::Type PanelId)
+{
+	MergeAssetSet[PanelId].Revision = NewRevision;
+	OnAssetChanged.ExecuteIfBound(PanelId, MergeAssetSet[PanelId]);
+}
+
+#undef LOCTEXT_NAMESPACE
diff --git a/Engine/Source/Developer/Merge/Private/SMergeAssetPickerView.h b/Engine/Source/Developer/Merge/Private/SMergeAssetPickerView.h
new file mode 100644
--- /dev/null
+++ b/Engine/Source/Developer/Merge/Private/SMergeAssetPickerView.h
@@ -0,0 +1,55 @@
+// Copyright 1998-2015 Epic Games, Inc. All Rights Reserved.
+#pragma once
+
+#include "IDiffControl.h"
+#include "BlueprintMergeData.h" // for IMergeControl
+#include "MergeUtils.h" // for EMergeAssetId
+
+class SMergeAssetPickerPanel;
+
+/*******************************************************************************
+ * FAssetRevisionInfo
+ ******************************************************************************/
+
+struct FAssetRevisionInfo
+{
+	FString AssetName;
+	FRevisionInfo Revision;
+};
+
+/*******************************************************************************
+ * SMergeAssetPickerView
+ ******************************************************************************/
+
+class SMergeAssetPickerView : public SCompoundWidget
+                      , public IDiffControl
+					  , public IMergeControl
+{
+	DECLARE_DELEGATE_TwoParams(FOnMergeAssetChanged, EMergeAssetId::Type, const FAssetRevisionInfo&)
+
+public:
+	SLATE_BEGIN_ARGS(SMergeAssetPickerView){}
+		SLATE_EVENT(FOnMergeAssetChanged, OnAssetChanged)
+	SLATE_END_ARGS()
+
+	void Construct(const FArguments InArgs, const FBlueprintMergeData& InData);
+
+private:
+	/** Implementation of IDiffControl: */
+	void NextDiff() override {}
+	void PrevDiff() override {}
+	bool HasNextDifference() const override { return false; }
+	bool HasPrevDifference() const override { return false; }
+
+	/** Implementation of IMergeControl */
+	void HighlightNextConflict() override {}
+	void HighlightPrevConflict() override {}
+	bool HasNextConflict() const override { return false; }
+	bool HasPrevConflict() const override { return false; }
+
+	void HandleAssetChange(const UObject* NewAsset, EMergeAssetId::Type PanelId);
+	void HandleRevisionChange(const FRevisionInfo& NewRevision, EMergeAssetId::Type PanelId);
+
+	FOnMergeAssetChanged OnAssetChanged;
+	FAssetRevisionInfo MergeAssetSet[EMergeAssetId::Count];
+};
diff --git a/Engine/Source/Editor/EditorStyle/Private/SlateEditorStyle.cpp b/Engine/Source/Editor/EditorStyle/Private/SlateEditorStyle.cpp
--- a/Engine/Source/Editor/EditorStyle/Private/SlateEditorStyle.cpp
+++ b/Engine/Source/Editor/EditorStyle/Private/SlateEditorStyle.cpp
@@ -4655,6 +4655,7 @@ void FSlateEditorStyle::FStyle::SetupPersonaStyle()
 		Set("BlueprintMerge.Cancel", new IMAGE_BRUSH("/Icons/LV_Remove", Icon16x16));
 		Set("BlueprintMerge.AcceptSource", new IMAGE_BRUSH("/Icons/AcceptMergeSource_40x", Icon16x16));
 		Set("BlueprintMerge.AcceptTarget", new IMAGE_BRUSH("/Icons/AcceptMergeTarget_40x", Icon16x16));
+		Set("BlueprintMerge.StartMerge", new IMAGE_BRUSH("/Icons/StartMerge_42x", Icon16x16));
 	}
 
 	// Play in editor / play in world
diff --git a/Engine/Source/Editor/Kismet/Private/SBlueprintEditorToolbar.cpp b/Engine/Source/Editor/Kismet/Private/SBlueprintEditorToolbar.cpp
--- a/Engine/Source/Editor/Kismet/Private/SBlueprintEditorToolbar.cpp
+++ b/Engine/Source/Editor/Kismet/Private/SBlueprintEditorToolbar.cpp
@@ -23,7 +23,7 @@
 #include "Engine/LevelScriptBlueprint.h"
 #include "Merge.h"
 #include "DesktopPlatformModule.h"
-#include "SThrobber.h"
+#include "SBlueprintRevisionMenu.h"
 
 #define LOCTEXT_NAMESPACE "KismetToolbar"
 
@@ -70,12 +70,7 @@ void FKismet2Menu::FillFileMenuBlueprintSection( FMenuBuilder& MenuBuilder, FBlu
 			LOCTEXT("BlueprintEditorDiffToolTip", "Diff against previous revisions"),
 			FOnGetContent::CreateStatic< FBlueprintEditor& >( &FKismet2Menu::MakeDiffMenu, Kismet),
 			FSlateIcon());
-
-		UBlueprint* BlueprintObj = Kismet.GetBlueprintObj();
-		if (BlueprintObj && IMerge::Get().PendingMerge( *BlueprintObj ) )
-		{
-			MenuBuilder.AddMenuEntry(FBlueprintEditorCommands::Get().BeginBlueprintMerge);
-		}
+		MenuBuilder.AddMenuEntry(FBlueprintEditorCommands::Get().BeginBlueprintMerge);
 	}
 	MenuBuilder.EndSection();
 
@@ -226,299 +221,83 @@ void FKismet2Menu::SetupBlueprintEditorMenu( TSharedPtr< FExtender > Extender, F
 	}
 }
 
-namespace EBlueprintDiffMenuState
+/** Delegate called to diff a specific revision with the current */
+static void OnDiffRevisionPicked(FRevisionInfo const& RevisionInfo, TWeakObjectPtr<UBlueprint> BlueprintObj)
 {
-	enum Type
-	{
-		NotQueried,
-		QueryInProgress,
-		Queried,
-	};
-}
-
-/**
- * Helper widget for managing & displaying source control history retrieval status
- */
-class SBlueprintDiffMenu : public SCompoundWidget
-{
-	SLATE_BEGIN_ARGS( SBlueprintDiffMenu ){}
-
-	SLATE_ARGUMENT(UBlueprint*, BlueprintObj)
-
-	SLATE_END_ARGS()
-
-	~SBlueprintDiffMenu()
+	if (BlueprintObj.IsValid())
 	{
-		// cancel any operation if this widget is destroyed while in progress
-		if(State == EBlueprintDiffMenuState::QueryInProgress)
-		{
-			ISourceControlProvider& SourceControlProvider = ISourceControlModule::Get().GetProvider();
-			if(UpdateStatusOperation.IsValid() && SourceControlProvider.CanCancelOperation(UpdateStatusOperation.ToSharedRef()))
-			{
-				SourceControlProvider.CancelOperation(UpdateStatusOperation.ToSharedRef());
-			}
-		}
-	}
+		bool const bIsLevelScriptBlueprint = FBlueprintEditorUtils::IsLevelScriptBlueprint(BlueprintObj.Get());
+		FString const Filename = SourceControlHelpers::PackageFilename(bIsLevelScriptBlueprint ? BlueprintObj.Get()->GetOuter()->GetPathName() : BlueprintObj.Get()->GetPathName());
 
-	void Construct(const FArguments& InArgs)
-	{
-		State = EBlueprintDiffMenuState::NotQueried;
-		BlueprintObj = InArgs._BlueprintObj;
-		bIsLevelScriptBlueprint = false;
-
-		ChildSlot	
-		[
-			SAssignNew(VerticalBox, SVerticalBox)
-			+SVerticalBox::Slot()
-			[
-				SNew(SBorder)
-				.Visibility(this, &SBlueprintDiffMenu::GetInProgressVisibility)
-				.BorderImage(FEditorStyle::GetBrush("Menu.Background"))
-				.Content()
-				[
-					SNew(SHorizontalBox)
-					+SHorizontalBox::Slot()
-					.AutoWidth()
-					.VAlign(VAlign_Center)
-					[
-						SNew(SThrobber)
-					]
-					+SHorizontalBox::Slot()
-					.AutoWidth()
-					.VAlign(VAlign_Center)
-					.Padding(2.0f, 0.0f, 4.0f, 0.0f)
-					[
-						SNew(STextBlock)
-						.Text(LOCTEXT("DiffMenuOperationInProgress", "Updating history..."))
-					]
-					+SHorizontalBox::Slot()
-					.FillWidth(1.0f)
-					.HAlign(HAlign_Right)
-					.VAlign(VAlign_Center)
-					[
-						SNew(SButton)
-						.Visibility(this, &SBlueprintDiffMenu::GetCancelButtonVisibility)
-						.OnClicked(this, &SBlueprintDiffMenu::OnCancelButtonClicked)
-						.VAlign(VAlign_Center)
-						.HAlign(HAlign_Center)
-						.Content()
-						[
-							SNew(STextBlock)
-							.Text(LOCTEXT("DiffMenuCancelButton", "Cancel"))
-						]
-					]
-				]
-			]
-		];
-
-		if(BlueprintObj.IsValid())
-		{
-			bIsLevelScriptBlueprint = FBlueprintEditorUtils::IsLevelScriptBlueprint(BlueprintObj.Get());
-			Filename = SourceControlHelpers::PackageFilename(bIsLevelScriptBlueprint ? BlueprintObj.Get()->GetOuter()->GetPathName() : BlueprintObj.Get()->GetPathName());
-
-			// make sure the history info is up to date
-			UpdateStatusOperation = ISourceControlOperation::Create<FUpdateStatus>();
-			UpdateStatusOperation->SetUpdateHistory(true);
-			ISourceControlModule::Get().GetProvider().Execute(UpdateStatusOperation.ToSharedRef(), Filename, EConcurrency::Asynchronous, FSourceControlOperationComplete::CreateSP(this, &SBlueprintDiffMenu::HandleSourceControlOperationComplete));
-
-			State = EBlueprintDiffMenuState::QueryInProgress;
-		}
-	}
-
-	/** Callback for when the source control operation is complete */
-	void HandleSourceControlOperationComplete( const FSourceControlOperationRef& InOperation, ECommandResult::Type InResult )
-	{
-		check(UpdateStatusOperation == InOperation);
+		ISourceControlProvider& SourceControlProvider = ISourceControlModule::Get().GetProvider();
 
-		if(InResult == ECommandResult::Succeeded)
+		// Get the SCC state
+		FSourceControlStatePtr SourceControlState = SourceControlProvider.GetState(Filename, EStateCacheUsage::Use);
+		if (SourceControlState.IsValid())
 		{
-			// get the cached state
-			ISourceControlProvider& SourceControlProvider = ISourceControlModule::Get().GetProvider();
-			FSourceControlStatePtr SourceControlState = SourceControlProvider.GetState(Filename, EStateCacheUsage::Use);
-
-			//Add pop-out menu for each revision
-			FMenuBuilder MenuBuilder(true, NULL);
-
-			MenuBuilder.BeginSection("AddDiffRevision", LOCTEXT("Revisions", "Revisions") );
-
-			if (SourceControlState.IsValid() && SourceControlState->GetHistorySize() > 0)
+			for (int32 HistoryIndex = 0; HistoryIndex < SourceControlState->GetHistorySize(); HistoryIndex++)
 			{
-				// Figure out the highest revision # (so we can label it "Depot")
-				int32 LatestRevision = 0;
-				for(int32 HistoryIndex = 0; HistoryIndex < SourceControlState->GetHistorySize(); HistoryIndex++)
+				TSharedPtr<ISourceControlRevision, ESPMode::ThreadSafe> Revision = SourceControlState->GetHistoryItem(HistoryIndex);
+				check(Revision.IsValid());
+				if (Revision->GetRevisionNumber() == RevisionInfo.Revision)
 				{
-					TSharedPtr<ISourceControlRevision, ESPMode::ThreadSafe> Revision = SourceControlState->GetHistoryItem(HistoryIndex);
-					if(Revision.IsValid() && Revision->GetRevisionNumber() > LatestRevision)
+					// Get the revision of this package from source control
+					FString PreviousTempPkgName;
+					if (Revision->Get(PreviousTempPkgName))
 					{
-						LatestRevision = Revision->GetRevisionNumber();
-					}
-				}
-
-				for(int32 HistoryIndex = 0; HistoryIndex < SourceControlState->GetHistorySize(); HistoryIndex++)
-				{
-					TSharedPtr<ISourceControlRevision, ESPMode::ThreadSafe> Revision = SourceControlState->GetHistoryItem(HistoryIndex);
-					if(Revision.IsValid())
-					{
-						FInternationalization& I18N = FInternationalization::Get();
-
-						FText Label = FText::Format( LOCTEXT("RevisionNumber", "Revision {0}"), FText::AsNumber( Revision->GetRevisionNumber(), NULL, I18N.GetInvariantCulture() ) );
-
-						FFormatNamedArguments Args;
-						Args.Add( TEXT("CheckInNumber"), FText::AsNumber( Revision->GetCheckInIdentifier(), NULL, I18N.GetInvariantCulture() ) );
-						Args.Add( TEXT("UserName"), FText::FromString( Revision->GetUserName() ) );
-						Args.Add( TEXT("DateTime"), FText::AsDate( Revision->GetDate() ) );
-						Args.Add( TEXT("ChanglistDescription"), FText::FromString( Revision->GetDescription() ) );
-						const FText ToolTip = FText::Format( LOCTEXT("RevisionToolTip", "CL #{CheckInNumber} {UserName} \n{DateTime} \n{ChanglistDescription}"), Args );
-						
-						if(LatestRevision == Revision->GetRevisionNumber())
-						{
-							Label = LOCTEXT("Depo", "Depot");
-						}
-
-						MenuBuilder.AddMenuEntry(TAttribute<FText>(Label), ToolTip, FSlateIcon(), 
-									FUIAction(FExecuteAction::CreateSP(this, &SBlueprintDiffMenu::DiffAgainstRevision, Revision->GetRevisionNumber())));
-					}
-				}
-			}
-			else
-			{
-				// Show 'empty' item in toolbar
-				MenuBuilder.AddMenuEntry( LOCTEXT("NoRevisonHistory", "No revisions found"), 
-					FText(), FSlateIcon(), FUIAction() );
-			}
+						// Forcibly disable compile on load in case we are loading old blueprints that might try to update/compile
+						TGuardValue<bool> DisableCompileOnLoad(GForceDisableBlueprintCompileOnLoad, true);
 
-			MenuBuilder.EndSection();
+						// Try and load that package
+						UPackage* PreviousTempPkg = LoadPackage(NULL, *PreviousTempPkgName, LOAD_None);
 
-			VerticalBox->AddSlot()
-			[
-				MenuBuilder.MakeWidget()
-			];
-		}
-
-		UpdateStatusOperation.Reset();
-		State = EBlueprintDiffMenuState::Queried;
-	}
-
-	/** Delegate called to diff a specific revision with the current */
-	void DiffAgainstRevision( int32 InOldRevision )
-	{
-		if(BlueprintObj.IsValid())
-		{
-			ISourceControlProvider& SourceControlProvider = ISourceControlModule::Get().GetProvider();
-
-			// Get the SCC state
-			FSourceControlStatePtr SourceControlState = SourceControlProvider.GetState(Filename, EStateCacheUsage::Use);
-			if(SourceControlState.IsValid())
-			{
-				for ( int32 HistoryIndex = 0; HistoryIndex < SourceControlState->GetHistorySize(); HistoryIndex++ )
-				{
-					TSharedPtr<ISourceControlRevision, ESPMode::ThreadSafe> Revision = SourceControlState->GetHistoryItem(HistoryIndex);
-					check(Revision.IsValid());
-					if(Revision->GetRevisionNumber() == InOldRevision)
-					{
-						// Get the revision of this package from source control
-						FString PreviousTempPkgName;
-						if(Revision->Get(PreviousTempPkgName))
+						if (PreviousTempPkg != NULL)
 						{
-							// Forcibly disable compile on load in case we are loading old blueprints that might try to update/compile
-							TGuardValue<bool> DisableCompileOnLoad(GForceDisableBlueprintCompileOnLoad, true);
-
-							// Try and load that package
-							UPackage* PreviousTempPkg = LoadPackage(NULL, *PreviousTempPkgName, LOAD_None);
+							UObject* PreviousAsset = NULL;
 
-							if(PreviousTempPkg != NULL)
+							// If its a levelscript blueprint, find the previous levelscript blueprint in the map
+							if (bIsLevelScriptBlueprint)
 							{
-								UObject* PreviousAsset = NULL;
+								TArray<UObject *> ObjectsInOuter;
+								GetObjectsWithOuter(PreviousTempPkg, ObjectsInOuter);
 
-								// If its a levelscript blueprint, find the previous levelscript blueprint in the map
-								if (bIsLevelScriptBlueprint)
+								// Look for the level script blueprint for this package
+								for (int32 Index = 0; Index < ObjectsInOuter.Num(); Index++)
 								{
-									TArray<UObject *> ObjectsInOuter;
-									GetObjectsWithOuter(PreviousTempPkg, ObjectsInOuter);
-						
-									// Look for the level script blueprint for this package
-									for( int32 Index = 0; Index < ObjectsInOuter.Num(); Index++ )
+									UObject* Obj = ObjectsInOuter[Index];
+									if (ULevelScriptBlueprint* ObjAsBlueprint = Cast<ULevelScriptBlueprint>(Obj))
 									{
-										UObject* Obj = ObjectsInOuter[Index];
-										if (ULevelScriptBlueprint* ObjAsBlueprint = Cast<ULevelScriptBlueprint>(Obj))
-										{
-											PreviousAsset = ObjAsBlueprint;
-											break;
-										}
+										PreviousAsset = ObjAsBlueprint;
+										break;
 									}
 								}
-								// otherwise its a normal Blueprint
-								else
-								{
-									FString PreviousAssetName = FPaths::GetBaseFilename(Filename, true);
-									PreviousAsset = FindObject<UObject>(PreviousTempPkg, *PreviousAssetName);
-								}
-
-								if(PreviousAsset != NULL)
-								{
-									FAssetToolsModule& AssetToolsModule = FModuleManager::LoadModuleChecked<FAssetToolsModule>(TEXT("AssetTools"));
-									FRevisionInfo OldRevision = {Revision->GetRevisionNumber(), Revision->GetCheckInIdentifier(), Revision->GetDate()};
-									FRevisionInfo CurrentRevision = {-1, Revision->GetCheckInIdentifier(), Revision->GetDate()};
-									AssetToolsModule.Get().DiffAssets(PreviousAsset, BlueprintObj.Get(), OldRevision, CurrentRevision );
-								}
 							}
+							// otherwise its a normal Blueprint
 							else
 							{
-								FMessageDialog::Open( EAppMsgType::Ok, NSLOCTEXT("SourceControl.HistoryWindow", "UnableToLoadAssets", "Unable to load assets to diff. Content may no longer be supported?"));
+								FString PreviousAssetName = FPaths::GetBaseFilename(Filename, true);
+								PreviousAsset = FindObject<UObject>(PreviousTempPkg, *PreviousAssetName);
+							}
+
+							if (PreviousAsset != NULL)
+							{
+								FAssetToolsModule& AssetToolsModule = FModuleManager::LoadModuleChecked<FAssetToolsModule>(TEXT("AssetTools"));
+								FRevisionInfo OldRevision = { Revision->GetRevisionNumber(), Revision->GetCheckInIdentifier(), Revision->GetDate() };
+								FRevisionInfo CurrentRevision = { -1, Revision->GetCheckInIdentifier(), Revision->GetDate() };
+								AssetToolsModule.Get().DiffAssets(PreviousAsset, BlueprintObj.Get(), OldRevision, CurrentRevision);
 							}
 						}
-						break;
+						else
+						{
+							FMessageDialog::Open(EAppMsgType::Ok, NSLOCTEXT("SourceControl.HistoryWindow", "UnableToLoadAssets", "Unable to load assets to diff. Content may no longer be supported?"));
+						}
 					}
+					break;
 				}
 			}
 		}
 	}
-
-	/** Delegate used to cancel a source control operation in progress */
-	FReply OnCancelButtonClicked() const
-	{
-		if(UpdateStatusOperation.IsValid())
-		{
-			ISourceControlProvider& SourceControlProvider = ISourceControlModule::Get().GetProvider();
-			SourceControlProvider.CancelOperation(UpdateStatusOperation.ToSharedRef());
-		}
-
-		return FReply::Handled();
-	}
-
-	/** Delegate used to determine the visibility of the cancel button */
-	EVisibility GetCancelButtonVisibility() const
-	{
-		ISourceControlProvider& SourceControlProvider = ISourceControlModule::Get().GetProvider();
-		return UpdateStatusOperation.IsValid() && SourceControlProvider.CanCancelOperation(UpdateStatusOperation.ToSharedRef()) ? EVisibility::Visible : EVisibility::Collapsed;
-	}
-
-	/** Delegate used to determine the visibility 'in progress' widgets */
-	EVisibility GetInProgressVisibility() const
-	{
-		return State == EBlueprintDiffMenuState::QueryInProgress ? EVisibility::Visible : EVisibility::Collapsed;
-	}
-
-private:
-
-	/** The source control operation in progress */
-	TSharedPtr<FUpdateStatus, ESPMode::ThreadSafe> UpdateStatusOperation;
-
-	/** The box we are using to display our menu */
-	TSharedPtr<SVerticalBox> VerticalBox;
-
-	/** The state of the SCC query */
-	EBlueprintDiffMenuState::Type State;
-
-	/** The Blueprint we are interested in */
-	TWeakObjectPtr<UBlueprint> BlueprintObj;
-
-	/** The name of the file we want to diff */
-	FString Filename;
-
-	/** We need to handle level script blueprints differently, so we cache this here */
-	bool bIsLevelScriptBlueprint;
-};
+}
 
 TSharedRef<SWidget> FKismet2Menu::MakeDiffMenu(FBlueprintEditor& Kismet)
 {
@@ -527,8 +306,10 @@ TSharedRef<SWidget> FKismet2Menu::MakeDiffMenu(FBlueprintEditor& Kismet)
 		UBlueprint* BlueprintObj = Kismet.GetBlueprintObj();
 		if(BlueprintObj)
 		{
+			TWeakObjectPtr<UBlueprint> BlueprintPtr = BlueprintObj;
 			// Add our async SCC task widget
-			return SNew(SBlueprintDiffMenu).BlueprintObj(BlueprintObj);
+			return SNew(SBlueprintRevisionMenu, BlueprintObj)
+				.OnRevisionSelected_Static(&OnDiffRevisionPicked, BlueprintPtr);
 		}
 		else
 		{
diff --git a/Engine/Source/Editor/Kismet/Private/SBlueprintRevisionMenu.cpp b/Engine/Source/Editor/Kismet/Private/SBlueprintRevisionMenu.cpp
new file mode 100644
--- /dev/null
+++ b/Engine/Source/Editor/Kismet/Private/SBlueprintRevisionMenu.cpp
@@ -0,0 +1,231 @@
+// Copyright 1998-2015 Epic Games, Inc. All Rights Reserved.
+
+
+#include "BlueprintEditorPrivatePCH.h"
+#include "SBlueprintRevisionMenu.h"
+#include "Developer/SourceControl/Public/ISourceControlModule.h"
+#include "SThrobber.h"
+
+#define LOCTEXT_NAMESPACE "SBlueprintRevisionMenu"
+
+/**  */
+namespace ESourceControlQueryState
+{
+	enum Type
+	{
+		NotQueried,
+		QueryInProgress,
+		Queried,
+	};
+}
+
+//------------------------------------------------------------------------------
+SBlueprintRevisionMenu::~SBlueprintRevisionMenu()
+{
+	// cancel any operation if this widget is destroyed while in progress
+	if (SourceControlQueryState == ESourceControlQueryState::QueryInProgress)
+	{
+		ISourceControlProvider& SourceControlProvider = ISourceControlModule::Get().GetProvider();
+		if (SourceControlQueryOp.IsValid() && SourceControlProvider.CanCancelOperation(SourceControlQueryOp.ToSharedRef()))
+		{
+			SourceControlProvider.CancelOperation(SourceControlQueryOp.ToSharedRef());
+		}
+	}
+}
+
+//------------------------------------------------------------------------------
+void SBlueprintRevisionMenu::Construct(const FArguments& InArgs, UBlueprint const* Blueprint)
+{
+	bIncludeLocalRevision = InArgs._bIncludeLocalRevision;
+	OnRevisionSelected = InArgs._OnRevisionSelected;
+
+	SourceControlQueryState = ESourceControlQueryState::NotQueried;
+
+	ChildSlot	
+	[
+		SAssignNew(MenuBox, SVerticalBox)
+		+SVerticalBox::Slot()
+		[
+			SNew(SBorder)
+			.Visibility(this, &SBlueprintRevisionMenu::GetInProgressVisibility)
+			.BorderImage(FEditorStyle::GetBrush("Menu.Background"))
+			.Content()
+			[
+				SNew(SHorizontalBox)
+				+SHorizontalBox::Slot()
+				.AutoWidth()
+				.VAlign(VAlign_Center)
+				[
+					SNew(SThrobber)
+				]
+				+SHorizontalBox::Slot()
+				.AutoWidth()
+				.VAlign(VAlign_Center)
+				.Padding(2.0f, 0.0f, 4.0f, 0.0f)
+				[
+					SNew(STextBlock)
+					.Text(LOCTEXT("DiffMenuOperationInProgress", "Updating history..."))
+				]
+				+SHorizontalBox::Slot()
+				.FillWidth(1.0f)
+				.HAlign(HAlign_Right)
+				.VAlign(VAlign_Center)
+				[
+					SNew(SButton)
+					.Visibility(this, &SBlueprintRevisionMenu::GetCancelButtonVisibility)
+					.OnClicked(this, &SBlueprintRevisionMenu::OnCancelButtonClicked)
+					.VAlign(VAlign_Center)
+					.HAlign(HAlign_Center)
+					.Content()
+					[
+						SNew(STextBlock)
+						.Text(LOCTEXT("DiffMenuCancelButton", "Cancel"))
+					]
+				]
+			]
+		]
+	];
+
+	if (Blueprint != nullptr)
+	{
+		bool const bIsLevelScriptBlueprint = FBlueprintEditorUtils::IsLevelScriptBlueprint(Blueprint);
+		Filename = SourceControlHelpers::PackageFilename(bIsLevelScriptBlueprint ? Blueprint->GetOuter()->GetPathName() : Blueprint->GetPathName());
+
+		// make sure the history info is up to date
+		SourceControlQueryOp = ISourceControlOperation::Create<FUpdateStatus>();
+		SourceControlQueryOp->SetUpdateHistory(true);
+		ISourceControlModule::Get().GetProvider().Execute(SourceControlQueryOp.ToSharedRef(), Filename, EConcurrency::Asynchronous, FSourceControlOperationComplete::CreateSP(this, &SBlueprintRevisionMenu::OnSourceControlQueryComplete));
+
+		SourceControlQueryState = ESourceControlQueryState::QueryInProgress;
+	}
+}
+
+//------------------------------------------------------------------------------
+EVisibility SBlueprintRevisionMenu::GetInProgressVisibility() const
+{
+	return (SourceControlQueryState == ESourceControlQueryState::QueryInProgress) ? EVisibility::Visible : EVisibility::Collapsed;
+}
+
+//------------------------------------------------------------------------------
+EVisibility SBlueprintRevisionMenu::GetCancelButtonVisibility() const
+{
+	ISourceControlProvider& SourceControlProvider = ISourceControlModule::Get().GetProvider();
+	return SourceControlQueryOp.IsValid() && SourceControlProvider.CanCancelOperation(SourceControlQueryOp.ToSharedRef()) ? EVisibility::Visible : EVisibility::Collapsed;
+}
+
+//------------------------------------------------------------------------------
+FReply SBlueprintRevisionMenu::OnCancelButtonClicked() const
+{
+	if (SourceControlQueryOp.IsValid())
+	{
+		ISourceControlProvider& SourceControlProvider = ISourceControlModule::Get().GetProvider();
+		SourceControlProvider.CancelOperation(SourceControlQueryOp.ToSharedRef());
+	}
+
+	return FReply::Handled();
+}
+
+//------------------------------------------------------------------------------
+void SBlueprintRevisionMenu::OnSourceControlQueryComplete(const FSourceControlOperationRef& InOperation, ECommandResult::Type InResult)
+{
+	check(SourceControlQueryOp == InOperation);
+
+
+	// Add pop-out menu for each revision
+	FMenuBuilder MenuBuilder(/*bInShouldCloseWindowAfterMenuSelection =*/true, /*InCommandList =*/NULL);
+
+	MenuBuilder.BeginSection("AddDiffRevision", LOCTEXT("Revisions", "Revisions"));
+	if (bIncludeLocalRevision)
+	{
+		FText const ToolTip = LOCTEXT("LocalRevisionToolTip", "The current copy you have saved to disk (locally)");
+
+		FOnRevisionSelected OnRevisionSelectedDelegate = OnRevisionSelected;
+		auto OnMenuItemSelected = [OnRevisionSelectedDelegate]()
+		{
+			OnRevisionSelectedDelegate.ExecuteIfBound(FRevisionInfo::InvalidRevision());
+		};
+
+		MenuBuilder.AddMenuEntry(LOCTEXT("LocalRevision", "Local"), ToolTip, FSlateIcon(),
+			FUIAction(FExecuteAction::CreateLambda(OnMenuItemSelected)));
+	}
+
+	if (InResult == ECommandResult::Succeeded)
+	{
+		// get the cached state
+		ISourceControlProvider& SourceControlProvider = ISourceControlModule::Get().GetProvider();
+		FSourceControlStatePtr SourceControlState = SourceControlProvider.GetState(Filename, EStateCacheUsage::Use);
+
+		if (SourceControlState.IsValid() && SourceControlState->GetHistorySize() > 0)
+		{
+			// Figure out the highest revision # (so we can label it "Depot")
+			int32 LatestRevision = 0;
+			for (int32 HistoryIndex = 0; HistoryIndex < SourceControlState->GetHistorySize(); HistoryIndex++)
+			{
+				TSharedPtr<ISourceControlRevision, ESPMode::ThreadSafe> Revision = SourceControlState->GetHistoryItem(HistoryIndex);
+				if (Revision.IsValid() && Revision->GetRevisionNumber() > LatestRevision)
+				{
+					LatestRevision = Revision->GetRevisionNumber();
+				}
+			}
+
+			for (int32 HistoryIndex = 0; HistoryIndex < SourceControlState->GetHistorySize(); HistoryIndex++)
+			{
+				TSharedPtr<ISourceControlRevision, ESPMode::ThreadSafe> Revision = SourceControlState->GetHistoryItem(HistoryIndex);
+				if (Revision.IsValid())
+				{
+					FInternationalization& I18N = FInternationalization::Get();
+
+					FText Label = FText::Format(LOCTEXT("RevisionNumber", "Revision {0}"), FText::AsNumber(Revision->GetRevisionNumber(), NULL, I18N.GetInvariantCulture()));
+
+					FFormatNamedArguments Args;
+					Args.Add(TEXT("CheckInNumber"), FText::AsNumber(Revision->GetCheckInIdentifier(), NULL, I18N.GetInvariantCulture()));
+					Args.Add(TEXT("UserName"), FText::FromString(Revision->GetUserName()));
+					Args.Add(TEXT("DateTime"), FText::AsDate(Revision->GetDate()));
+					Args.Add(TEXT("ChanglistDescription"), FText::FromString(Revision->GetDescription()));
+					const FText ToolTip = FText::Format(LOCTEXT("RevisionToolTip", "CL #{CheckInNumber} {UserName} \n{DateTime} \n{ChanglistDescription}"), Args);
+
+					if (LatestRevision == Revision->GetRevisionNumber())
+					{
+						Label = LOCTEXT("Depo", "Depot");
+					}
+
+					FRevisionInfo RevisionInfo = { 
+						Revision->GetRevisionNumber(), 
+						Revision->GetCheckInIdentifier(), 
+						Revision->GetDate() 
+					};
+					FOnRevisionSelected OnRevisionSelectedDelegate = OnRevisionSelected;
+					auto OnMenuItemSelected = [RevisionInfo, OnRevisionSelectedDelegate]()
+					{
+						OnRevisionSelectedDelegate.ExecuteIfBound(RevisionInfo);
+					};
+					MenuBuilder.AddMenuEntry( TAttribute<FText>(Label), ToolTip, FSlateIcon(),
+						FUIAction(FExecuteAction::CreateLambda(OnMenuItemSelected)) );
+				}
+			}
+		}
+		else if (!bIncludeLocalRevision)
+		{
+			// Show 'empty' item in toolbar
+			MenuBuilder.AddMenuEntry(LOCTEXT("NoRevisonHistory", "No revisions found"),
+				FText(), FSlateIcon(), FUIAction());
+		}
+	}
+	else if (!bIncludeLocalRevision)
+	{
+		// Show 'empty' item in toolbar
+		MenuBuilder.AddMenuEntry(LOCTEXT("NoRevisonHistory", "No revisions found"),
+			FText(), FSlateIcon(), FUIAction());
+	}
+
+	MenuBuilder.EndSection();
+	MenuBox->AddSlot() 
+	[
+		MenuBuilder.MakeWidget()
+	];
+
+	SourceControlQueryOp.Reset();
+	SourceControlQueryState = ESourceControlQueryState::Queried;
+}
+
+#undef LOCTEXT_NAMESPACE
diff --git a/Engine/Source/Editor/Kismet/Public/SBlueprintRevisionMenu.h b/Engine/Source/Editor/Kismet/Public/SBlueprintRevisionMenu.h
new file mode 100644
--- /dev/null
+++ b/Engine/Source/Editor/Kismet/Public/SBlueprintRevisionMenu.h
@@ -0,0 +1,45 @@
+// Copyright 1998-2015 Epic Games, Inc. All Rights Reserved.
+#pragma once
+
+struct FRevisionInfo;
+
+class KISMET_API SBlueprintRevisionMenu : public SCompoundWidget
+{
+	DECLARE_DELEGATE_OneParam(FOnRevisionSelected, FRevisionInfo const&)
+
+public:
+	SLATE_BEGIN_ARGS(SBlueprintRevisionMenu)
+		: _bIncludeLocalRevision(false)
+	{}
+		SLATE_ARGUMENT(bool, bIncludeLocalRevision)
+		SLATE_EVENT(FOnRevisionSelected, OnRevisionSelected)
+	SLATE_END_ARGS()
+
+	~SBlueprintRevisionMenu();
+
+	void Construct(const FArguments& InArgs, UBlueprint const* Blueprint);
+
+private: 
+	/** Delegate used to determine the visibility 'in progress' widgets */
+	EVisibility GetInProgressVisibility() const;
+	/** Delegate used to determine the visibility of the cancel button */
+	EVisibility GetCancelButtonVisibility() const;
+
+	/** Delegate used to cancel a source control operation in progress */
+	FReply OnCancelButtonClicked() const;
+	/** Callback for when the source control operation is complete */
+	void OnSourceControlQueryComplete(const FSourceControlOperationRef& InOperation, ECommandResult::Type InResult);
+
+	/**  */
+	bool bIncludeLocalRevision;
+	/**  */
+	FOnRevisionSelected OnRevisionSelected;
+	/** The name of the file we want revision info for */
+	FString Filename;
+	/** The box we are using to display our menu */
+	TSharedPtr<SVerticalBox> MenuBox;
+	/** The source control operation in progress */
+	TSharedPtr<FUpdateStatus, ESPMode::ThreadSafe> SourceControlQueryOp;
+	/** The state of the SCC query */
+	uint32 SourceControlQueryState;
+};
