Fixing some mono/non-unity build errors.

diff --git a/Engine/Source/Developer/Merge/Private/SMergeAssetPickerView.cpp b/Engine/Source/Developer/Merge/Private/SMergeAssetPickerView.cpp
--- a/Engine/Source/Developer/Merge/Private/SMergeAssetPickerView.cpp
+++ b/Engine/Source/Developer/Merge/Private/SMergeAssetPickerView.cpp
@@ -19,9 +19,9 @@ class SMergeAssetPickerPanel : public SCompoundWidget
 
 public:
 	SLATE_BEGIN_ARGS(SMergeAssetPickerPanel)
-		: _SelectedAsset(nullptr)
+		: _AssetClass(UBlueprint::StaticClass())
+		, _SelectedAsset(nullptr)
 		, _RevisionInfo(FRevisionInfo::InvalidRevision())
-		, _AssetClass(UBlueprint::StaticClass())
 	{}
 		SLATE_ARGUMENT(TSubclassOf<UBlueprint>, AssetClass)
 		SLATE_ARGUMENT(const UBlueprint*, SelectedAsset)
diff --git a/Engine/Source/Editor/Kismet/Private/SBlueprintRevisionMenu.cpp b/Engine/Source/Editor/Kismet/Private/SBlueprintRevisionMenu.cpp
--- a/Engine/Source/Editor/Kismet/Private/SBlueprintRevisionMenu.cpp
+++ b/Engine/Source/Editor/Kismet/Private/SBlueprintRevisionMenu.cpp
@@ -5,6 +5,8 @@
 #include "SBlueprintRevisionMenu.h"
 #include "Developer/SourceControl/Public/ISourceControlModule.h"
 #include "SThrobber.h"
+#include "BlueprintEditorUtils.h"
+#include "IAssetTypeActions.h"	// for FRevisionInfo
 
 #define LOCTEXT_NAMESPACE "SBlueprintRevisionMenu"
 
