Non-unity fix.

diff --git a/Engine/Source/Developer/Merge/Private/MergeUtils.h b/Engine/Source/Developer/Merge/Private/MergeUtils.h
--- a/Engine/Source/Developer/Merge/Private/MergeUtils.h
+++ b/Engine/Source/Developer/Merge/Private/MergeUtils.h
@@ -5,7 +5,8 @@
 #include "ISourceControlState.h" // for FSourceControlStatePtr
 
 // Forward declarations
-class ISourceControlRevision;
+class  ISourceControlRevision;
+struct FRevisionInfo;
 
 /**  */
 namespace EMergeAssetId
