Fix non-unity in MergeUtils.cpp

diff --git a/Engine/Source/Developer/Merge/Private/MergeUtils.cpp b/Engine/Source/Developer/Merge/Private/MergeUtils.cpp
--- a/Engine/Source/Developer/Merge/Private/MergeUtils.cpp
+++ b/Engine/Source/Developer/Merge/Private/MergeUtils.cpp
@@ -5,6 +5,8 @@
 #include "ISourceControlModule.h"
 #include "SourceControlHelpers.h"
 #include "NotificationManager.h"
+#include "Developer/AssetTools/Public/IAssetTypeActions.h"
+#include "Runtime/Slate/Public/Widgets/Notifications/SNotificationList.h"
 
 #define LOCTEXT_NAMESPACE "MergeToolUtils"
 
