Fix CIS.

diff --git a/Engine/Source/Developer/Merge/Private/SMergeDetailsView.cpp b/Engine/Source/Developer/Merge/Private/SMergeDetailsView.cpp
--- a/Engine/Source/Developer/Merge/Private/SMergeDetailsView.cpp
+++ b/Engine/Source/Developer/Merge/Private/SMergeDetailsView.cpp
@@ -159,7 +159,7 @@ void SMergeDetailsView::Construct(const FArguments InArgs
 		bool bLocalMatchesBase = !bDoneLocal && !bDoneBase && BaseVisibleProperties[IterBase] == LocalVisibleProperties[IterLocal];
 		bool bRemoteMatchesBase = !bDoneRemote && !bDoneBase && BaseVisibleProperties[IterBase] == RemoteVisibleProperties[IterRemote];
 
-		if( bRemoteMatchesBase && bLocalMatchesBase ||
+		if( (bRemoteMatchesBase && bLocalMatchesBase) ||
 			(bDoneLocal && bRemoteMatchesBase ) ||
 			(bDoneRemote && bLocalMatchesBase ) )
 		{
