Initializing SMergeAssetPickerView::MergeAssetSet, could not start a merge if the revision was changed but the asset was not.

diff --git a/Engine/Source/Developer/Merge/Private/SMergeAssetPickerView.cpp b/Engine/Source/Developer/Merge/Private/SMergeAssetPickerView.cpp
--- a/Engine/Source/Developer/Merge/Private/SMergeAssetPickerView.cpp
+++ b/Engine/Source/Developer/Merge/Private/SMergeAssetPickerView.cpp
@@ -379,6 +379,13 @@ void SMergeAssetPickerPanel::OnRevisionSelected(FRevisionInfo const& RevisionInf
 //------------------------------------------------------------------------------
 void SMergeAssetPickerView::Construct(const FArguments InArgs, const FBlueprintMergeData& InData)
 {
+	MergeAssetSet[EMergeAssetId::MergeRemote].AssetName = InData.BlueprintRemote->GetOutermost()->GetName();
+	MergeAssetSet[EMergeAssetId::MergeRemote].Revision = InData.RevisionRemote;
+	MergeAssetSet[EMergeAssetId::MergeBase].AssetName = InData.BlueprintBase->GetOutermost()->GetName();
+	MergeAssetSet[EMergeAssetId::MergeBase].Revision = InData.RevisionBase;
+	MergeAssetSet[EMergeAssetId::MergeLocal].AssetName = InData.BlueprintLocal->GetOutermost()->GetName();
+	MergeAssetSet[EMergeAssetId::MergeLocal].Revision = InData.RevisionLocal;
+	
 	OnAssetChanged = InArgs._OnAssetChanged;
 
 	ChildSlot
