Fixes for blueprint merge tool when run on assets that aren't conflictng (!IsActivelyMerging when merge is initiated)

diff --git a/Engine/Source/Developer/Merge/Private/SBlueprintMerge.cpp b/Engine/Source/Developer/Merge/Private/SBlueprintMerge.cpp
--- a/Engine/Source/Developer/Merge/Private/SBlueprintMerge.cpp
+++ b/Engine/Source/Developer/Merge/Private/SBlueprintMerge.cpp
@@ -166,8 +166,6 @@ void SBlueprintMerge::Construct(const FArguments InArgs, const FBlueprintMergeDa
 		, FSlateIcon(FEditorStyle::GetStyleSetName(), "BlueprintMerge.Cancel")
 	);
 
-	DifferencesTreeView = DiffTreeView::CreateTreeView(&MasterDifferencesList);
-
 	TSharedRef<SWidget> Overlay = SNew(SHorizontalBox);
 	if( IsActivelyMerging() )
 	{
@@ -214,11 +212,8 @@ void SBlueprintMerge::Construct(const FArguments InArgs, const FBlueprintMergeDa
 			+ SSplitter::Slot()
 			.Value(.2f)
 			[
-				SNew(SBorder)
+				SAssignNew(TreeViewContainer, SBorder)
 				.BorderImage(FEditorStyle::GetBrush("ToolPanel.GroupBorder"))
-				[
-					DifferencesTreeView.ToSharedRef()
-				]
 			]
 			+ SSplitter::Slot()
 			.Value(.8f)
@@ -270,12 +265,12 @@ void SBlueprintMerge::PrevDiff()
 
 bool SBlueprintMerge::HasNextDiff() const
 {
-	return DiffTreeView::HasNextDifference( DifferencesTreeView.ToSharedRef(), RealDifferences );
+	return DifferencesTreeView.IsValid() && DiffTreeView::HasNextDifference( DifferencesTreeView.ToSharedRef(), RealDifferences );
 }
 
 bool SBlueprintMerge::HasPrevDiff() const
 {
-	return DiffTreeView::HasPrevDifference(DifferencesTreeView.ToSharedRef(), RealDifferences );
+	return DifferencesTreeView.IsValid() && DiffTreeView::HasPrevDifference(DifferencesTreeView.ToSharedRef(), RealDifferences );
 }
 
 void SBlueprintMerge::NextConflict()
@@ -290,16 +285,19 @@ void SBlueprintMerge::PrevConflict()
 
 bool SBlueprintMerge::HasNextConflict() const
 {
-	return DiffTreeView::HasNextDifference(DifferencesTreeView.ToSharedRef(), MergeConflicts);
+	return DifferencesTreeView.IsValid() && DiffTreeView::HasNextDifference(DifferencesTreeView.ToSharedRef(), MergeConflicts);
 }
 
 bool SBlueprintMerge::HasPrevConflict() const
 {
-	return DiffTreeView::HasPrevDifference(DifferencesTreeView.ToSharedRef(), MergeConflicts);
+	return DifferencesTreeView.IsValid() && DiffTreeView::HasPrevDifference(DifferencesTreeView.ToSharedRef(), MergeConflicts);
 }
 
 void SBlueprintMerge::OnStartMerge()
 {
+	DifferencesTreeView = DiffTreeView::CreateTreeView(&MasterDifferencesList);
+	TreeViewContainer->SetContent(DifferencesTreeView.ToSharedRef());
+
 	if (Data.BlueprintRemote == nullptr)
 	{
 		Data.BlueprintRemote = Cast<UBlueprint>(FMergeToolUtils::LoadRevision(RemotePath, Data.RevisionRemote));
diff --git a/Engine/Source/Developer/Merge/Private/SBlueprintMerge.h b/Engine/Source/Developer/Merge/Private/SBlueprintMerge.h
--- a/Engine/Source/Developer/Merge/Private/SBlueprintMerge.h
+++ b/Engine/Source/Developer/Merge/Private/SBlueprintMerge.h
@@ -101,6 +101,9 @@ private:
 	bool bIsPickingAssets;
 	FOnMergeResolved OnMergeResolved;
 
+	/** Container widget for the treeview of differences: */
+	TSharedPtr<SBorder> TreeViewContainer;
+
 	/** Treeview to display all differences collected across all panels: */
 	TSharedPtr< STreeView< TSharedPtr< class FBlueprintDifferenceTreeEntry > > > DifferencesTreeView;
 
diff --git a/Engine/Source/Developer/Merge/Private/SMergeGraphView.cpp b/Engine/Source/Developer/Merge/Private/SMergeGraphView.cpp
--- a/Engine/Source/Developer/Merge/Private/SMergeGraphView.cpp
+++ b/Engine/Source/Developer/Merge/Private/SMergeGraphView.cpp
@@ -52,6 +52,8 @@ struct FMergeGraphRowEntry
 	UEdGraphPin* RemotePin;
 
 	FLinearColor DisplayColor;
+
+	bool bHasConflicts;
 };
 
 struct FMergeGraphEntry
@@ -181,8 +183,8 @@ static TArray< FMergeGraphEntry > GenerateDiffListItems(const FBlueprintRevPair&
 								continue;
 							}
 
-							ConflictMap[&RemoteDifference] = ConflictingDifference;
-							ConflictMap[ConflictingDifference] = &RemoteDifference;
+							ConflictMap.Add(&RemoteDifference, ConflictingDifference);
+							ConflictMap.Add(ConflictingDifference, &RemoteDifference);
 						}
 					}
 
@@ -212,9 +214,10 @@ static TArray< FMergeGraphEntry > GenerateDiffListItems(const FBlueprintRevPair&
 							, Difference.Pin1 /*UEdGraphPin* BasePin*/
 							, Difference.Pin2 /*UEdGraphPin* RemotePin*/
 							, Difference.DisplayColor
+							, ConflictingDifference ? true : false
 						};
 
-						GraphEntry.bAnyConflics |= NewEntry.LocalNode && NewEntry.RemoteNode;
+						GraphEntry.bAnyConflics |= NewEntry.bHasConflicts;
 						GraphEntry.Changes.Push( NewEntry );
 					}
 
@@ -236,6 +239,7 @@ static TArray< FMergeGraphEntry > GenerateDiffListItems(const FBlueprintRevPair&
 								, Difference.Pin1 /*UEdGraphPin* BasePin*/
 								, nullptr
 								, Difference.DisplayColor
+								, false
 							};
 
 							GraphEntry.Changes.Push(NewEntry);
