Nonunity fix

diff --git a/Engine/Source/Developer/Merge/Private/SMergeGraphView.cpp b/Engine/Source/Developer/Merge/Private/SMergeGraphView.cpp
--- a/Engine/Source/Developer/Merge/Private/SMergeGraphView.cpp
+++ b/Engine/Source/Developer/Merge/Private/SMergeGraphView.cpp
@@ -1,6 +1,7 @@
 // Copyright 1998-2015 Epic Games, Inc. All Rights Reserved.
 
 #include "MergePrivatePCH.h"
+#include "DiffUtils.h"
 #include "DiffResults.h"
 #include "GraphDiffControl.h"
 #include "SDockTab.h"
