Fix for behavior that I can only explain via this compiler bug (or something similar)
 - https://connect.microsoft.com/VisualStudio/feedback/details/800364/initializer-list-calls-object-destructor-twice
Also fixed possible infinite loop

diff --git a/Engine/Source/Developer/Merge/Private/SMergeDetailsView.cpp b/Engine/Source/Developer/Merge/Private/SMergeDetailsView.cpp
--- a/Engine/Source/Developer/Merge/Private/SMergeDetailsView.cpp
+++ b/Engine/Source/Developer/Merge/Private/SMergeDetailsView.cpp
@@ -191,6 +191,10 @@ void SMergeDetailsView::Construct(const FArguments InArgs
 			AddPropertiesOrdered(LocalVisibleProperties[IterLocal], LocalDifferingPropertiesSet, DifferingProperties);
 			++IterLocal;
 		}
+		else if( bDoneLocal && bDoneRemote && !bDoneBase )
+		{
+			++IterBase;
+		}
 
 		bDoneRemote = IterRemote == RemoteVisibleProperties.Num();
 		bDoneBase = IterBase == BaseVisibleProperties.Num();
diff --git a/Engine/Source/Developer/Merge/Private/SMergeTreeView.cpp b/Engine/Source/Developer/Merge/Private/SMergeTreeView.cpp
--- a/Engine/Source/Developer/Merge/Private/SMergeTreeView.cpp
+++ b/Engine/Source/Developer/Merge/Private/SMergeTreeView.cpp
@@ -146,11 +146,16 @@ void SMergeTreeView::Construct(const FArguments InArgs
 				{
 					bAnyConflict = true;
 
+					FSCSMergeEntry Entry = 
+					{ 
+						FText::Format(NSLOCTEXT("SMergeTreeView", "ConflictIdentifier", "CONFLICT: {0} conflicts with {1}"), DiffViewUtils::SCSDiffMessage(Remote, RemoteLabel), DiffViewUtils::SCSDiffMessage(Local, LocalLabel)),
+						Remote.TreeIdentifier,
+						Remote.DiffType == ETreeDiffType::NODE_PROPERTY_CHANGED ? Remote.PropertyDiff.Identifier : Local.PropertyDiff.Identifier,
+						true 
+					};
+
 					// create a tree entry that describes both the local and remote change..
-					Entries.Push( {FText::Format( NSLOCTEXT("SMergeTreeView", "ConflictIdentifier", "CONFLICT: {0} conflicts with {1}"), DiffViewUtils::SCSDiffMessage( Remote, RemoteLabel ), DiffViewUtils::SCSDiffMessage( Local, LocalLabel ) ),
-									Remote.TreeIdentifier,
-									Remote.DiffType == ETreeDiffType::NODE_PROPERTY_CHANGED ? Remote.PropertyDiff.Identifier : Local.PropertyDiff.Identifier,
-									true });
+					Entries.Push( Entry );
 
 					++RemoteIter;
 					++LocalIter;
@@ -166,19 +171,27 @@ void SMergeTreeView::Construct(const FArguments InArgs
 
 		if( Local && ( !Remote || SortTreePredicate( *Local, *Remote ) ) )
 		{
-			Entries.Push( {	DiffViewUtils::SCSDiffMessage(*Local, LocalLabel),
-							Local->TreeIdentifier,
-							Local->PropertyDiff.Identifier,
-							false } );
+			FSCSMergeEntry Entry =
+			{ 
+				DiffViewUtils::SCSDiffMessage(*Local, LocalLabel),
+				Local->TreeIdentifier,
+				Local->PropertyDiff.Identifier,
+				false 
+			};
+			Entries.Push( Entry );
 
 			++LocalIter;
 		}
 		else
 		{
-			Entries.Push( { DiffViewUtils::SCSDiffMessage(*Remote, RemoteLabel),
-							Remote->TreeIdentifier,
-							Remote->PropertyDiff.Identifier,
-							false });
+			FSCSMergeEntry Entry =
+			{ 
+				DiffViewUtils::SCSDiffMessage(*Remote, RemoteLabel),
+				Remote->TreeIdentifier,
+				Remote->PropertyDiff.Identifier,
+				false 
+			};
+			Entries.Push( Entry );
 
 			++RemoteIter;
 		}
