[UE-10016] Fix BP merge tool crash when attempting to view component diffs.

change summary:
- switched the SCSViews array type to be a TSharedRef<FSCSDiff> in SMergeTreeView; previously the FSCSDiff instance was passed to TArray::Push() as a temporary, which i think didn't increment the refcount of FSCSDiff's internal widget ptrs on the shallow copy to the new element.

diff --git a/Engine/Source/Developer/Merge/Private/SMergeTreeView.cpp b/Engine/Source/Developer/Merge/Private/SMergeTreeView.cpp
--- a/Engine/Source/Developer/Merge/Private/SMergeTreeView.cpp
+++ b/Engine/Source/Developer/Merge/Private/SMergeTreeView.cpp
@@ -19,25 +19,25 @@ void SMergeTreeView::Construct(const FArguments InArgs
 	// EMergeParticipant::Remote
 	{
 		SCSViews.Push(
-			FSCSDiff(InData.BlueprintRemote)
+			MakeShareable(new FSCSDiff(InData.BlueprintRemote))
 			);
 	}
 	// EMergeParticipant::Base
 	{
 		SCSViews.Push(
-			FSCSDiff(InData.BlueprintBase)
+			MakeShareable(new FSCSDiff(InData.BlueprintBase))
 			);
 	}
 	// EMergeParticipant::Local
 	{
 		SCSViews.Push(
-			FSCSDiff(InData.BlueprintLocal)
+			MakeShareable(new FSCSDiff(InData.BlueprintLocal))
 			);
 	}
 
-	TArray< FSCSResolvedIdentifier > RemoteHierarchy = GetRemoteView().GetDisplayedHierarchy();
-	TArray< FSCSResolvedIdentifier > BaseHierarchy = GetBaseView().GetDisplayedHierarchy();
-	TArray< FSCSResolvedIdentifier > LocalHierarchy = GetLocalView().GetDisplayedHierarchy();
+	TArray< FSCSResolvedIdentifier > RemoteHierarchy = GetRemoteView()->GetDisplayedHierarchy();
+	TArray< FSCSResolvedIdentifier > BaseHierarchy = GetBaseView()->GetDisplayedHierarchy();
+	TArray< FSCSResolvedIdentifier > LocalHierarchy = GetLocalView()->GetDisplayedHierarchy();
 
 	FSCSDiffRoot RemoteDifferingProperties;
 	DiffUtils::CompareUnrelatedSCS(InData.BlueprintBase, BaseHierarchy, InData.BlueprintRemote, RemoteHierarchy, RemoteDifferingProperties );
@@ -254,15 +254,15 @@ void SMergeTreeView::Construct(const FArguments InArgs
 		SNew(SSplitter)
 			+ SSplitter::Slot()
 			[
-				GetRemoteView().TreeWidget()
+				GetRemoteView()->TreeWidget()
 			]
 		+ SSplitter::Slot()
 			[
-				GetBaseView().TreeWidget()
+				GetBaseView()->TreeWidget()
 			]
 		+ SSplitter::Slot()
 			[
-				GetLocalView().TreeWidget()
+				GetLocalView()->TreeWidget()
 			]
 	];
 }
@@ -271,21 +271,21 @@ void SMergeTreeView::HighlightDifference(FSCSIdentifier TreeIdentifier, FPropert
 {
 	for (auto& View : SCSViews)
 	{
-		View.HighlightProperty(TreeIdentifier.Name, Property);
+		View->HighlightProperty(TreeIdentifier.Name, Property);
 	}
 }
 
-FSCSDiff& SMergeTreeView::GetRemoteView()
+TSharedRef<FSCSDiff>& SMergeTreeView::GetRemoteView()
 {
 	return SCSViews[EMergeParticipant::Remote];
 }
 
-FSCSDiff& SMergeTreeView::GetBaseView()
+TSharedRef<FSCSDiff>& SMergeTreeView::GetBaseView()
 {
 	return SCSViews[EMergeParticipant::Base];
 }
 
-FSCSDiff& SMergeTreeView::GetLocalView()
+TSharedRef<FSCSDiff>& SMergeTreeView::GetLocalView()
 {
 	return SCSViews[EMergeParticipant::Local];
 }
diff --git a/Engine/Source/Developer/Merge/Private/SMergeTreeView.h b/Engine/Source/Developer/Merge/Private/SMergeTreeView.h
--- a/Engine/Source/Developer/Merge/Private/SMergeTreeView.h
+++ b/Engine/Source/Developer/Merge/Private/SMergeTreeView.h
@@ -22,12 +22,12 @@ public:
 private:
 	void HighlightDifference( FSCSIdentifier VarName, FPropertySoftPath Property );
 
-	FSCSDiff& GetRemoteView();
-	FSCSDiff& GetBaseView();
-	FSCSDiff& GetLocalView();
+	TSharedRef<FSCSDiff>& GetRemoteView();
+	TSharedRef<FSCSDiff>& GetBaseView();
+	TSharedRef<FSCSDiff>& GetLocalView();
 
 	FBlueprintMergeData Data;
-	TArray< FSCSDiff > SCSViews;
+	TArray< TSharedRef<FSCSDiff> > SCSViews;
 
 	FSCSDiffRoot MergeConflicts;
 	int CurrentMergeConflict;
