Fix shadowed variables

diff --git a/Engine/Plugins/2D/Paper2D/Source/Paper2D/Private/AlphaBitmap.h b/Engine/Plugins/2D/Paper2D/Source/Paper2D/Private/AlphaBitmap.h
--- a/Engine/Plugins/2D/Paper2D/Source/Paper2D/Private/AlphaBitmap.h
+++ b/Engine/Plugins/2D/Paper2D/Source/Paper2D/Private/AlphaBitmap.h
@@ -7,10 +7,10 @@
 
 struct FAlphaBitmap
 {
-	FAlphaBitmap(UTexture2D* SourceTexture, uint8 DefaultValue = 0)
+	FAlphaBitmap(UTexture2D* SourceTexture, uint8 InDefaultValue = 0)
 		: Width(0)
 		, Height(0)
-		, DefaultValue(DefaultValue)
+		, DefaultValue(InDefaultValue)
 	{
 		ExtractFromTexture(SourceTexture);
 	}
@@ -67,10 +67,10 @@ struct FAlphaBitmap
 	}
 
 	// Create an empty bitmap
-	FAlphaBitmap(int Width, int Height, int DefaultValue = 0)
-		: Width(Width)
-		, Height(Height)
-		, DefaultValue(DefaultValue)
+	FAlphaBitmap(int InWidth, int InHeight, int InDefaultValue = 0)
+		: Width(InWidth)
+		, Height(InHeight)
+		, DefaultValue(InDefaultValue)
 	{
 		RawData.SetNumZeroed(Width * Height);
 		ClearToDefaultValue();
@@ -203,12 +203,12 @@ struct FAlphaBitmap
 	}
 
 	// Sets pixels to 1 marking the rectangle outline
-	void DrawRectOutline(int StartX, int StartY, int Width, int Height)
+	void DrawRectOutline(int StartX, int StartY, int InWidth, int InHeight)
 	{
 		const int X0 = StartX;
 		const int Y0 = StartY;
-		const int X1 = StartX + Width - 1;
-		const int Y1 = StartY + Height - 1;
+		const int X1 = StartX + InWidth - 1;
+		const int Y1 = StartY + InHeight - 1;
 		for (int Y = Y0; Y <= Y1; ++Y)
 		{
 			SetPixel(X0, Y, 1);
@@ -221,12 +221,12 @@ struct FAlphaBitmap
 		}
 	}
 
-	void FillRect(int StartX, int StartY, int Width, int Height)
+	void FillRect(int StartX, int StartY, int InWidth, int InHeight)
 	{
 		const int X0 = StartX;
 		const int Y0 = StartY;
-		const int X1 = StartX + Width - 1;
-		const int Y1 = StartY + Height - 1;
+		const int X1 = StartX + InWidth - 1;
+		const int Y1 = StartY + InHeight - 1;
 		for (int Y = Y0; Y <= Y1; ++Y)
 		{
 			for (int X = X0; X <= X1; ++X)
@@ -323,15 +323,15 @@ struct FAlphaBitmap
 	}
 
 	// Checks the image to determine if it is suitable for opaque, masked or translucent rendering
-	void AnalyzeImage(int32 StartX, int32 StartY, int32 Width, int32 Height, bool& bOutHasZeros, bool& bOutHasIntermediateValues)
+	void AnalyzeImage(int32 StartX, int32 StartY, int32 InWidth, int32 InHeight, bool& bOutHasZeros, bool& bOutHasIntermediateValues)
 	{
 		bOutHasZeros = false;
 		bOutHasIntermediateValues = false;
 
 		const int32 X0 = StartX;
 		const int32 Y0 = StartY;
-		const int32 X1 = StartX + Width;
-		const int32 Y1 = StartY + Height;
+		const int32 X1 = StartX + InWidth;
+		const int32 Y1 = StartY + InHeight;
 		for (int Y = Y0; Y < Y1; ++Y)
 		{
 			for (int X = X0; X < X1; ++X)
diff --git a/Engine/Plugins/2D/Paper2D/Source/Paper2D/Private/PaperRenderSceneProxy.cpp b/Engine/Plugins/2D/Paper2D/Source/Paper2D/Private/PaperRenderSceneProxy.cpp
--- a/Engine/Plugins/2D/Paper2D/Source/Paper2D/Private/PaperRenderSceneProxy.cpp
+++ b/Engine/Plugins/2D/Paper2D/Source/Paper2D/Private/PaperRenderSceneProxy.cpp
@@ -488,8 +488,6 @@ void FPaperRenderSceneProxy::GetBatchMesh(const FSceneView* View, bool bUseOverr
 			BatchElement.PrimitiveUniformBuffer = GetUniformBuffer();
 			BatchElement.NumPrimitives = VertexArray.Vertices.Num() / 3;
 
-			const bool bIsWireframeView = View->Family->EngineShowFlags.Wireframe;
-
 			Collector.AddMesh(ViewIndex, Mesh);
 		}
 	}
diff --git a/Engine/Plugins/2D/Paper2D/Source/Paper2D/Private/PaperSprite.cpp b/Engine/Plugins/2D/Paper2D/Source/Paper2D/Private/PaperSprite.cpp
--- a/Engine/Plugins/2D/Paper2D/Source/Paper2D/Private/PaperSprite.cpp
+++ b/Engine/Plugins/2D/Paper2D/Source/Paper2D/Private/PaperSprite.cpp
@@ -1572,10 +1572,10 @@ void UPaperSprite::SetRotated(bool bRotated)
 	RebuildCollisionData();
 }
 
-void UPaperSprite::SetPivotMode(ESpritePivotMode::Type PivotMode, FVector2D CustomTextureSpacePivot)
+void UPaperSprite::SetPivotMode(ESpritePivotMode::Type InPivotMode, FVector2D InCustomTextureSpacePivot)
 {
-	this->PivotMode = PivotMode;
-	this->CustomPivotPoint = CustomTextureSpacePivot;
+	PivotMode = InPivotMode;
+	CustomPivotPoint = InCustomTextureSpacePivot;
 	RebuildRenderData();
 	RebuildCollisionData();
 }
@@ -1825,7 +1825,7 @@ void UPaperSprite::ValidateSocketNames()
 	TSet<FName> SocketNames;
 	struct Local
 	{
-		static FName GetUniqueName(const TSet<FName>& SocketNames, FName Name)
+		static FName GetUniqueName(const TSet<FName>& InSocketNames, FName Name)
 		{
 			int Counter = Name.GetNumber();
 			FName TestName;
@@ -1833,7 +1833,7 @@ void UPaperSprite::ValidateSocketNames()
 			{
 				TestName = Name;
 				TestName.SetNumber(++Counter);
-			} while (SocketNames.Contains(TestName));
+			} while (InSocketNames.Contains(TestName));
 
 			return TestName;
 		}
diff --git a/Engine/Plugins/2D/Paper2D/Source/Paper2D/Private/Terrain/PaperTerrainComponent.cpp b/Engine/Plugins/2D/Paper2D/Source/Paper2D/Private/Terrain/PaperTerrainComponent.cpp
--- a/Engine/Plugins/2D/Paper2D/Source/Paper2D/Private/Terrain/PaperTerrainComponent.cpp
+++ b/Engine/Plugins/2D/Paper2D/Source/Paper2D/Private/Terrain/PaperTerrainComponent.cpp
@@ -403,10 +403,10 @@ void UPaperTerrainComponent::OnSplineEdited()
 			TArray<const UPaperSprite*> ValidBodies;
 			TArray<float> ValidBodyWidths;
 
-			int32 GenerateBodyIndex(FRandomStream& RandomStream) const
+			int32 GenerateBodyIndex(FRandomStream& InRandomStream) const
 			{
 				check(ValidBodies.Num() > 0);
-				return RandomStream.GetUnsignedInt() % ValidBodies.Num();
+				return InRandomStream.GetUnsignedInt() % ValidBodies.Num();
 			}
 		};
 
diff --git a/Engine/Plugins/2D/Paper2D/Source/Paper2DEditor/Private/Atlasing/PaperAtlasTextureHelpers.h b/Engine/Plugins/2D/Paper2D/Source/Paper2DEditor/Private/Atlasing/PaperAtlasTextureHelpers.h
--- a/Engine/Plugins/2D/Paper2D/Source/Paper2DEditor/Private/Atlasing/PaperAtlasTextureHelpers.h
+++ b/Engine/Plugins/2D/Paper2D/Source/Paper2DEditor/Private/Atlasing/PaperAtlasTextureHelpers.h
@@ -20,11 +20,13 @@ static void CopyTextureData(const uint8* Source, uint8* Dest, uint32 SizeX, uint
 // Read sprite data into the 4 bytes per pixel target array
 static bool ReadSpriteTexture(TArray<uint8>& TargetBuffer, FIntPoint& OutSpriteSize, UPaperSprite* Sprite)
 {
-	const int32 BytesPerPixel = 4;
 	const FVector2D SpriteSizeFloat = Sprite->GetSourceSize();
 	const FIntPoint SpriteSize(FMath::TruncToInt(SpriteSizeFloat.X), FMath::TruncToInt(SpriteSizeFloat.Y));
-	TargetBuffer.Empty();
-	TargetBuffer.AddZeroed(SpriteSize.X * SpriteSize.Y * BytesPerPixel);
+	{
+		const int32 BytesPerPixel = 4;
+		TargetBuffer.Empty();
+		TargetBuffer.AddZeroed(SpriteSize.X * SpriteSize.Y * BytesPerPixel);
+	}
 
 	check(Sprite->GetSourceTexture());
 	FTextureSource& SourceData = Sprite->GetSourceTexture()->Source;
diff --git a/Engine/Plugins/2D/Paper2D/Source/Paper2DEditor/Private/PaperFlipbookHelpers.cpp b/Engine/Plugins/2D/Paper2D/Source/Paper2DEditor/Private/PaperFlipbookHelpers.cpp
--- a/Engine/Plugins/2D/Paper2D/Source/Paper2DEditor/Private/PaperFlipbookHelpers.cpp
+++ b/Engine/Plugins/2D/Paper2D/Source/Paper2DEditor/Private/PaperFlipbookHelpers.cpp
@@ -156,8 +156,7 @@ void FPaperFlipbookHelpers::ExtractFlipbooksFromSprites(TMap<FString, TArray<UPa
 	OutSpriteFlipbookMap.GetKeys(Keys);
 	for (auto SpriteName : Keys)
 	{
-		TArray<UPaperSprite*>& Sprites = OutSpriteFlipbookMap[SpriteName];
-		Sprites.Sort(FSpriteSortPredicate());
+		OutSpriteFlipbookMap[SpriteName].Sort(FSpriteSortPredicate());
 	}
 
 	// Create a flipbook from all remaining sprites
diff --git a/Engine/Plugins/2D/Paper2D/Source/Paper2DEditor/Private/SpriteEditor/SpriteEditorSelections.h b/Engine/Plugins/2D/Paper2D/Source/Paper2DEditor/Private/SpriteEditor/SpriteEditorSelections.h
--- a/Engine/Plugins/2D/Paper2D/Source/Paper2DEditor/Private/SpriteEditor/SpriteEditorSelections.h
+++ b/Engine/Plugins/2D/Paper2D/Source/Paper2DEditor/Private/SpriteEditor/SpriteEditorSelections.h
@@ -213,17 +213,14 @@ public:
 	{
 	}
 
-	virtual bool IsValidInEditor(UPaperSprite* Sprite, bool bRenderData) const
+	virtual bool IsValidInEditor(UPaperSprite* Sprite, bool bInRenderData) const
 	{
-		if (SpritePtr == Sprite && this->bRenderData == bRenderData)
+		if (Sprite && SpritePtr == Sprite && bRenderData == bInRenderData)
 		{
-			if (UPaperSprite* Sprite = SpritePtr.Get())
+			FSpritePolygonCollection& Geometry = bRenderData ? Sprite->RenderGeometry : Sprite->CollisionGeometry;
+			if (Geometry.Polygons.IsValidIndex(PolygonIndex) && Geometry.Polygons[PolygonIndex].Vertices.IsValidIndex(VertexIndex))
 			{
-				FSpritePolygonCollection& Geometry = bRenderData ? Sprite->RenderGeometry : Sprite->CollisionGeometry;
-				if (Geometry.Polygons.IsValidIndex(PolygonIndex) && Geometry.Polygons[PolygonIndex].Vertices.IsValidIndex(VertexIndex))
-				{
-					return true;
-				}
+				return true;
 			}
 		}
 		return false;
diff --git a/Engine/Plugins/2D/Paper2D/Source/PaperJsonImporter/Private/PaperJsonSpriteSheetImporter.cpp b/Engine/Plugins/2D/Paper2D/Source/PaperJsonImporter/Private/PaperJsonSpriteSheetImporter.cpp
--- a/Engine/Plugins/2D/Paper2D/Source/PaperJsonImporter/Private/PaperJsonSpriteSheetImporter.cpp
+++ b/Engine/Plugins/2D/Paper2D/Source/PaperJsonImporter/Private/PaperJsonSpriteSheetImporter.cpp
@@ -217,7 +217,6 @@ static bool ParseFramesFromSpriteArray(const TArray<TSharedPtr<FJsonValue>>& Arr
 	bool bLoadedSuccessfully = true;
 
 	// Parse all of the frames
-	int32 FrameCount = 0;
 	for (int32 FrameCount = 0; FrameCount < ArrayBlock.Num(); ++FrameCount)
 	{
 		GWarn->StatusUpdate(FrameCount, ArrayBlock.Num(), NSLOCTEXT("Paper2D", "PaperJsonImporterFactory_ParsingSprites", "Parsing Sprite Frames"));
@@ -322,10 +321,10 @@ FPaperJsonSpriteSheetImporter::FPaperJsonSpriteSheetImporter() : bIsReimporting(
 {
 }
 
-void FPaperJsonSpriteSheetImporter::SetReimportData(const FString& ExistingTextureName, UTexture2D* ExistingTexture, const TArray<FString>& ExistingSpriteNames, const TArray< TAssetPtr<class UPaperSprite> >& ExistingSpriteAssetPtrs)
+void FPaperJsonSpriteSheetImporter::SetReimportData(const FString& InExistingTextureName, UTexture2D* InExistingTexture, const TArray<FString>& ExistingSpriteNames, const TArray< TAssetPtr<class UPaperSprite> >& ExistingSpriteAssetPtrs)
 {
-	this->ExistingTextureName = ExistingTextureName;
-	this->ExistingTexture = ExistingTexture;
+	ExistingTextureName = InExistingTextureName;
+	ExistingTexture = InExistingTexture;
 
 	check(ExistingSpriteNames.Num() == ExistingSpriteAssetPtrs.Num());
 	if (ExistingSpriteNames.Num() == ExistingSpriteAssetPtrs.Num())
diff --git a/Engine/Plugins/2D/Paper2D/Source/PaperJsonImporter/Private/PaperSpriteSheetAssetTypeActions.cpp b/Engine/Plugins/2D/Paper2D/Source/PaperJsonImporter/Private/PaperSpriteSheetAssetTypeActions.cpp
--- a/Engine/Plugins/2D/Paper2D/Source/PaperJsonImporter/Private/PaperSpriteSheetAssetTypeActions.cpp
+++ b/Engine/Plugins/2D/Paper2D/Source/PaperJsonImporter/Private/PaperSpriteSheetAssetTypeActions.cpp
@@ -86,28 +86,31 @@ void FPaperSpriteSheetAssetTypeActions::ExecuteCreateFlipbooks(TArray<TWeakObjec
 			bool useSpriteNames = (SpriteSheet->SpriteNames.Num() == SpriteSheet->Sprites.Num());
 
 			// Create a list of sprites and sprite names to feed into paper flipbook helpers
-			TArray<UPaperSprite*> Sprites;
-			TArray<FString> SpriteNames;
+			TMap<FString, TArray<UPaperSprite*> > SpriteFlipbookMap;
 
-			for (int SpriteIndex = 0; SpriteIndex < SpriteSheet->Sprites.Num(); ++SpriteIndex)
 			{
-				auto SpriteAssetPtr = SpriteSheet->Sprites[SpriteIndex];
-				FStringAssetReference SpriteStringRef = SpriteAssetPtr.ToStringReference();
-				UPaperSprite* Sprite = nullptr;
-				if (!SpriteStringRef.ToString().IsEmpty())
-				{
-					Sprite = Cast<UPaperSprite>(StaticLoadObject(UPaperSprite::StaticClass(), nullptr, *SpriteStringRef.ToString(), nullptr, LOAD_None, nullptr));
-				}
-				if (Sprite != nullptr)
+				TArray<UPaperSprite*> Sprites;
+				TArray<FString> SpriteNames;
+
+				for (int SpriteIndex = 0; SpriteIndex < SpriteSheet->Sprites.Num(); ++SpriteIndex)
 				{
-					const FString SpriteName = useSpriteNames ? SpriteSheet->SpriteNames[SpriteIndex] : Sprite->GetName();
-					Sprites.Add(Sprite);
-					SpriteNames.Add(SpriteName);
+					auto SpriteAssetPtr = SpriteSheet->Sprites[SpriteIndex];
+					FStringAssetReference SpriteStringRef = SpriteAssetPtr.ToStringReference();
+					UPaperSprite* Sprite = nullptr;
+					if (!SpriteStringRef.ToString().IsEmpty())
+					{
+						Sprite = Cast<UPaperSprite>(StaticLoadObject(UPaperSprite::StaticClass(), nullptr, *SpriteStringRef.ToString(), nullptr, LOAD_None, nullptr));
+					}
+					if (Sprite != nullptr)
+					{
+						const FString SpriteName = useSpriteNames ? SpriteSheet->SpriteNames[SpriteIndex] : Sprite->GetName();
+						Sprites.Add(Sprite);
+						SpriteNames.Add(SpriteName);
+					}
 				}
-			}
 
-			TMap<FString, TArray<UPaperSprite*> > SpriteFlipbookMap;
-			FPaperFlipbookHelpers::ExtractFlipbooksFromSprites(/*out*/SpriteFlipbookMap, Sprites, SpriteNames);
+				FPaperFlipbookHelpers::ExtractFlipbooksFromSprites(/*out*/SpriteFlipbookMap, Sprites, SpriteNames);
+			}
 
 			// Create one flipbook for every grouped flipbook name
 			if (SpriteFlipbookMap.Num() > 0)
diff --git a/Engine/Source/Developer/LogVisualizer/Private/STimeline.cpp b/Engine/Source/Developer/LogVisualizer/Private/STimeline.cpp
--- a/Engine/Source/Developer/LogVisualizer/Private/STimeline.cpp
+++ b/Engine/Source/Developer/LogVisualizer/Private/STimeline.cpp
@@ -120,7 +120,7 @@ bool STimeline::IsEntryHidden(const FVisualLogDevice::FVisualLogEntryItem& Entry
 	return HiddenEntries.Contains(&EntryItem);
 }
 
-void STimeline::HandleLogVisualizerSettingChanged(FName Name)
+void STimeline::HandleLogVisualizerSettingChanged(FName InName)
 {
 	UpdateVisibility();
 }
diff --git a/Engine/Source/Developer/LogVisualizer/Private/STimelineBar.cpp b/Engine/Source/Developer/LogVisualizer/Private/STimelineBar.cpp
--- a/Engine/Source/Developer/LogVisualizer/Private/STimelineBar.cpp
+++ b/Engine/Source/Developer/LogVisualizer/Private/STimelineBar.cpp
@@ -481,13 +481,13 @@ int32 STimelineBar::OnPaint(const FPaintArgs& Args, const FGeometry& AllottedGeo
 
 	if (BestItemIndex != INDEX_NONE && TimelineOwner.Pin()->IsSelected())
 	{
-		auto &Entries = TimelineOwner.Pin()->GetEntries();
+		const auto &BestItemEntry = TimelineOwner.Pin()->GetEntries()[BestItemIndex];
 		if (BestItemIndex != CurrentItemIndex)
 		{
-			FLogVisualizer::Get().GetVisualLoggerEvents().OnItemSelectionChanged.ExecuteIfBound(Entries[BestItemIndex]);
+			FLogVisualizer::Get().GetVisualLoggerEvents().OnItemSelectionChanged.ExecuteIfBound(BestItemEntry);
 		}
 
-		float CurrentTime = Entries[BestItemIndex].Entry.TimeStamp;
+		float CurrentTime = BestItemEntry.Entry.TimeStamp;
 		float LinePos = (CurrentTime - LocalViewRange.GetLowerBoundValue()) * PixelsPerInput;
 
 		FSlateDrawElement::MakeBox(
diff --git a/Engine/Source/Developer/LogVisualizer/Private/STimelinesContainer.cpp b/Engine/Source/Developer/LogVisualizer/Private/STimelinesContainer.cpp
--- a/Engine/Source/Developer/LogVisualizer/Private/STimelinesContainer.cpp
+++ b/Engine/Source/Developer/LogVisualizer/Private/STimelinesContainer.cpp
@@ -9,13 +9,13 @@
 
 #define LOCTEXT_NAMESPACE "STimelinesContainer"
 
-TSharedRef<SWidget> STimelinesContainer::MakeTimeline(TSharedPtr<class SVisualLoggerView> VisualLoggerView, TSharedPtr<class FVisualLoggerTimeSliderController> TimeSliderController, const FVisualLogDevice::FVisualLogEntryItem& Entry)
+TSharedRef<SWidget> STimelinesContainer::MakeTimeline(TSharedPtr<class SVisualLoggerView> InVisualLoggerView, TSharedPtr<class FVisualLoggerTimeSliderController> InTimeSliderController, const FVisualLogDevice::FVisualLogEntryItem& Entry)
 {
 	TSharedPtr<STimeline> NewTimeline;
 
 	ContainingBorder->AddSlot()
 		[
-			SAssignNew(NewTimeline, STimeline, VisualLoggerView, TimeSliderController, SharedThis(this), Entry)
+			SAssignNew(NewTimeline, STimeline, InVisualLoggerView, InTimeSliderController, SharedThis(this), Entry)
 			.OnItemSelectionChanged(FLogVisualizer::Get().GetVisualLoggerEvents().OnItemSelectionChanged)
 			.OnGetMenuContent(this, &STimelinesContainer::GetRightClickMenuContent)
 		];
diff --git a/Engine/Source/Developer/LogVisualizer/Private/SVisualLogger.cpp b/Engine/Source/Developer/LogVisualizer/Private/SVisualLogger.cpp
--- a/Engine/Source/Developer/LogVisualizer/Private/SVisualLogger.cpp
+++ b/Engine/Source/Developer/LogVisualizer/Private/SVisualLogger.cpp
@@ -321,7 +321,7 @@ void SVisualLogger::HandleTabManagerPersistLayout(const TSharedRef<FTabManager::
 	FLayoutSaveRestore::SaveToConfig(GEditorLayoutIni, LayoutToSave);
 }
 
-void SVisualLogger::FillFileMenu(FMenuBuilder& MenuBuilder, const TSharedPtr<FTabManager> TabManager)
+void SVisualLogger::FillFileMenu(FMenuBuilder& MenuBuilder, const TSharedPtr<FTabManager> InTabManager)
 {
 	MenuBuilder.BeginSection("LogFile", LOCTEXT("FileMenu", "Log File"));
 	{
diff --git a/Engine/Source/Developer/LogVisualizer/Private/SVisualLoggerFilters.cpp b/Engine/Source/Developer/LogVisualizer/Private/SVisualLoggerFilters.cpp
--- a/Engine/Source/Developer/LogVisualizer/Private/SVisualLoggerFilters.cpp
+++ b/Engine/Source/Developer/LogVisualizer/Private/SVisualLoggerFilters.cpp
@@ -141,10 +141,9 @@ void SVisualLoggerFilters::GraphFilterCategoryClicked(FName MenuCategory)
 	if (GraphFilters.Contains(MenuCategory))
 	{
 		bool bChanged = false;
-		auto &Filters = GraphFilters[MenuCategory];
-		for (int32 Index = 0; Index < Filters.Num(); ++Index)
+		for (const FString &Filter : GraphFilters[MenuCategory])
 		{
-			FString GraphFilterName = MenuCategory.ToString() + TEXT("$") + Filters[Index];
+			const FString GraphFilterName = MenuCategory.ToString() + TEXT("$") + Filter;
 			FCategoryFiltersManager::Get().GetCategory(GraphFilterName).Enabled = bNewSet;
 			bChanged = true;
 		}
@@ -161,10 +160,9 @@ bool SVisualLoggerFilters::IsGraphFilterCategoryInUse(FName MenuCategory) const
 	bool bInUse = false;
 	if (GraphFilters.Contains(MenuCategory))
 	{
-		auto &Filters = GraphFilters[MenuCategory];
-		for (int32 Index = 0; Index < Filters.Num(); ++Index)
+		for (const FString& Filter : GraphFilters[MenuCategory])
 		{
-			FString GraphFilterName = MenuCategory.ToString() + TEXT("$") + Filters[Index];
+			const FString GraphFilterName = MenuCategory.ToString() + TEXT("$") + Filter;
 			bInUse |= FCategoryFiltersManager::Get().GetCategory(GraphFilterName).Enabled;
 		}
 	}
@@ -200,14 +198,13 @@ void SVisualLoggerFilters::FilterByTypeClicked(FName InGraphName, FName InDataNa
 	if (GraphFilters.Contains(InGraphName))
 	{
 		bool bChanged = false;
-		auto &Filters = GraphFilters[InGraphName];
-		for (int32 Index = 0; Index < Filters.Num(); ++Index)
+		for (const FString& Filter : GraphFilters[InGraphName])
 		{
-			if (Filters[Index] == InDataName.ToString())
+			if (Filter == InDataName.ToString())
 			{
-				FString GraphFilterName = InGraphName.ToString() + TEXT("$") + Filters[Index];
-				FCategoryFilter& Filter = FCategoryFiltersManager::Get().GetCategory(GraphFilterName);
-				Filter.Enabled = !Filter.Enabled;
+				const FString GraphFilterName = InGraphName.ToString() + TEXT("$") + Filter;
+				FCategoryFilter& CategoryFilter = FCategoryFiltersManager::Get().GetCategory(GraphFilterName);
+				CategoryFilter.Enabled = !CategoryFilter.Enabled;
 				bChanged = true;
 			}
 		}
@@ -223,12 +220,11 @@ bool SVisualLoggerFilters::IsAssetTypeActionsInUse(FName InGraphName, FName InDa
 {
 	if (GraphFilters.Contains(InGraphName))
 	{
-		auto &Filters = GraphFilters[InGraphName];
-		for (int32 Index = 0; Index < Filters.Num(); ++Index)
+		for (const FString& Filter : GraphFilters[InGraphName])
 		{
-			if (Filters[Index] == InDataName.ToString())
+			if (Filter == InDataName.ToString())
 			{
-				FString GraphFilterName = InGraphName.ToString() + TEXT("$") + Filters[Index];
+				const FString GraphFilterName = InGraphName.ToString() + TEXT("$") + Filter;
 				return FCategoryFiltersManager::Get().GetCategory(GraphFilterName).Enabled;
 			}
 		}
diff --git a/Engine/Source/Developer/LogVisualizer/Private/VisualLoggerCanvasRenderer.cpp b/Engine/Source/Developer/LogVisualizer/Private/VisualLoggerCanvasRenderer.cpp
--- a/Engine/Source/Developer/LogVisualizer/Private/VisualLoggerCanvasRenderer.cpp
+++ b/Engine/Source/Developer/LogVisualizer/Private/VisualLoggerCanvasRenderer.cpp
@@ -217,7 +217,6 @@ void FVisualLoggerCanvasRenderer::DrawHistogramGraphs(class UCanvas* Canvas, cla
 	}
 
 	const float GoldenRatioConjugate = 0.618033988749895f;
-	int32 GraphIndex = 0;
 	if (CollectedGraphs.Num() > 0)
 	{
 		const FColor GraphsBackgroundColor = ULogVisualizerSettings::StaticClass()->GetDefaultObject<ULogVisualizerSettings>()->GraphsBackgroundColor;
@@ -281,9 +280,9 @@ void FVisualLoggerCanvasRenderer::DrawHistogramGraphs(class UCanvas* Canvas, cla
 				HistogramGraph->GetGraphLine(LineIndex)->LeftExtreme = LinesIt->Value.LeftExtreme;
 				HistogramGraph->GetGraphLine(LineIndex)->RightExtreme = LinesIt->Value.RightExtreme;
 
-				int32 DummyY, CurrentX;
-				StringSize(Font, CurrentX, DummyY, *LinesIt->Value.DataName.ToString());
-				MaxStringSize = CurrentX > MaxStringSize ? CurrentX : MaxStringSize;
+				int32 DummyY, StringSizeX;
+				StringSize(Font, StringSizeX, DummyY, *LinesIt->Value.DataName.ToString());
+				MaxStringSize = StringSizeX > MaxStringSize ? StringSizeX : MaxStringSize;
 
 				++LineIndex;
 			}
diff --git a/Engine/Source/Developer/LogVisualizer/Private/VisualLoggerRenderingActor.cpp b/Engine/Source/Developer/LogVisualizer/Private/VisualLoggerRenderingActor.cpp
--- a/Engine/Source/Developer/LogVisualizer/Private/VisualLoggerRenderingActor.cpp
+++ b/Engine/Source/Developer/LogVisualizer/Private/VisualLoggerRenderingActor.cpp
@@ -203,11 +203,13 @@ void AVisualLoggerRenderingActor::OnItemSelectionChanged(const FVisualLogDevice:
 	AddDebugRendering();
 #endif
 
-	const float Length = 100;
-	const FVector DirectionNorm = FVector(0, 0, 1).GetSafeNormal();
-	FVector YAxis, ZAxis;
-	DirectionNorm.FindBestAxisVectors(YAxis, ZAxis);
-	Cones.Add(FDebugRenderSceneProxy::FCone(FScaleMatrix(FVector(Length)) * FMatrix(DirectionNorm, YAxis, ZAxis, Entry->Location), 5, 5, FColor::Red));
+	{
+		const float Length = 100;
+		const FVector DirectionNorm = FVector(0, 0, 1).GetSafeNormal();
+		FVector YAxis, ZAxis;
+		DirectionNorm.FindBestAxisVectors(YAxis, ZAxis);
+		Cones.Add(FDebugRenderSceneProxy::FCone(FScaleMatrix(FVector(Length)) * FMatrix(DirectionNorm, YAxis, ZAxis, Entry->Location), 5, 5, FColor::Red));
+	}
 
 	if (LogEntriesPath.Num())
 	{
diff --git a/Engine/Source/Developer/Merge/Private/SMergeDetailsView.cpp b/Engine/Source/Developer/Merge/Private/SMergeDetailsView.cpp
--- a/Engine/Source/Developer/Merge/Private/SMergeDetailsView.cpp
+++ b/Engine/Source/Developer/Merge/Private/SMergeDetailsView.cpp
@@ -53,12 +53,12 @@ void SMergeDetailsView::Construct(const FArguments InArgs
 			2. Iterate properties in remote, add any new properties to DifferingProperties
 			3. Iterate properties in local, add any new properties to DifferingProperties, if they have not already been added
 	*/
-	const auto AddPropertiesOrdered = [](FPropertySoftPath const& Property, const FPropertySoftPathSet& DifferingProperties, TArray<FPropertySoftPath>& ResultingProperties)
+	const auto AddPropertiesOrdered = [](FPropertySoftPath const& Property, const FPropertySoftPathSet& InDifferingProperties, TArray<FPropertySoftPath>& ResultingProperties)
 	{
 		// contains check here is O(n), so we're needlessly n^2:
 		if (!ResultingProperties.Contains(Property))
 		{
-			if (DifferingProperties.Contains(Property))
+			if (InDifferingProperties.Contains(Property))
 			{
 				ResultingProperties.Add(Property);
 			}
@@ -105,19 +105,19 @@ void SMergeDetailsView::Construct(const FArguments InArgs
 	};
 
 	const auto ProcessPotentialDifference = []( const FPropertySoftPath& PropertyIdentifier
-											, TArray<FSingleObjectDiffEntry> const& RemoteDifferences
-											, TArray<FSingleObjectDiffEntry> const& LocalDifferences
+											, TArray<FSingleObjectDiffEntry> const& InRemoteDifferences
+											, TArray<FSingleObjectDiffEntry> const& InLocalDifferences
 											, TArray< FDiffPair >& OutProcessedDifferences
 											, bool& bOutAnyConflict )
 	{
 		const FText RemoteLabel = NSLOCTEXT("SMergeDetailsView", "RemoteLabel", "Remote");
 		const FText LocalLabel = NSLOCTEXT("SMergeDetailsView", "LocalLabel", "Local");
 
-		const auto FindDiffering = [](TArray< FSingleObjectDiffEntry > const& InDifferences, const FPropertySoftPath& PropertyIdentifier) -> const FSingleObjectDiffEntry*
+		const auto FindDiffering = [](TArray< FSingleObjectDiffEntry > const& InDifferences, const FPropertySoftPath& InPropertyIdentifier) -> const FSingleObjectDiffEntry*
 		{
 			for (const auto& Difference : InDifferences)
 			{
-				if (Difference.Identifier == PropertyIdentifier)
+				if (Difference.Identifier == InPropertyIdentifier)
 				{
 					return &Difference;
 				}
@@ -125,8 +125,8 @@ void SMergeDetailsView::Construct(const FArguments InArgs
 			return nullptr;
 		};
 
-		const FSingleObjectDiffEntry* RemoteDiffering = FindDiffering(RemoteDifferences, PropertyIdentifier);
-		const FSingleObjectDiffEntry* LocalDiffering = FindDiffering(LocalDifferences, PropertyIdentifier);
+		const FSingleObjectDiffEntry* RemoteDiffering = FindDiffering(InRemoteDifferences, PropertyIdentifier);
+		const FSingleObjectDiffEntry* LocalDiffering = FindDiffering(InLocalDifferences, PropertyIdentifier);
 		if (RemoteDiffering && LocalDiffering)
 		{
 			// conflicting change:
@@ -215,9 +215,9 @@ void SMergeDetailsView::Construct(const FArguments InArgs
 				.ColorAndOpacity(Entry.bConflicted ? DiffViewUtils::Conflicting() : DiffViewUtils::Differs());
 		};
 
-		const auto FocusDetailsDifferenceEntry = []( FPropertySoftPath PropertyIdentifier, SMergeDetailsView* Parent, FOnMergeNodeSelected SelectionCallback )
+		const auto FocusDetailsDifferenceEntry = []( FPropertySoftPath PropertyIdentifier, SMergeDetailsView* Parent, FOnMergeNodeSelected InSelectionCallback )
 		{
-			SelectionCallback.ExecuteIfBound();
+			InSelectionCallback.ExecuteIfBound();
 			Parent->HighlightDifference(PropertyIdentifier);
 		};
 
@@ -239,10 +239,10 @@ void SMergeDetailsView::Construct(const FArguments InArgs
 		}
 	}
 
-	const auto ForwardSelection = [](FOnMergeNodeSelected SelectionCallback)
+	const auto ForwardSelection = [](FOnMergeNodeSelected InSelectionCallback)
 	{
 		// This allows the owning control to focus the correct tab (or do whatever else it likes):
-		SelectionCallback.ExecuteIfBound();
+		InSelectionCallback.ExecuteIfBound();
 	};
 
 	TSharedPtr<FBlueprintDifferenceTreeEntry> Category = FBlueprintDifferenceTreeEntry::CreateDefaultsCategoryEntryForMerge(FOnDiffEntryFocused::CreateStatic(ForwardSelection, SelectionCallback), Children, RemoteDifferences.Num() != 0, LocalDifferences.Num() != 0, bAnyConflict);
diff --git a/Engine/Source/Developer/Merge/Private/SMergeGraphView.cpp b/Engine/Source/Developer/Merge/Private/SMergeGraphView.cpp
--- a/Engine/Source/Developer/Merge/Private/SMergeGraphView.cpp
+++ b/Engine/Source/Developer/Merge/Private/SMergeGraphView.cpp
@@ -397,9 +397,9 @@ void SMergeGraphView::Construct(const FArguments InArgs
 						.ColorAndOpacity(Color);
 			};
 
-			const auto SelectGraphNode = [](FOnMergeNodeSelected SelectionCallback, FMergeGraphRowEntry DiffEntry, SMergeGraphView* Parent)
+			const auto SelectGraphNode = [](FOnMergeNodeSelected InSelectionCallback, FMergeGraphRowEntry DiffEntry, SMergeGraphView* Parent)
 			{
-				SelectionCallback.ExecuteIfBound();
+				InSelectionCallback.ExecuteIfBound();
 				Parent->HighlightEntry( DiffEntry );
 			};
 
@@ -416,7 +416,7 @@ void SMergeGraphView::Construct(const FArguments InArgs
 			}
 		}
 
-		const auto Widget = []( const FMergeGraphEntry* Difference ) -> TSharedRef<SWidget>
+		const auto Widget = []( const FMergeGraphEntry* InDifference ) -> TSharedRef<SWidget>
 		{
 			// blue indicates added, red indicates changed, yellow indicates removed, white indicates no change:
 			const auto ComputeColor = [](const bool bAnyConflicts, const bool bAnyDifferences) -> FLinearColor
@@ -446,26 +446,26 @@ void SMergeGraphView::Construct(const FArguments InArgs
 					];
 			};
 
-			FLinearColor RemoteColor = ComputeColor(Difference->bAnyConflics, Difference->bRemoteDifferences);
-			FLinearColor BaseColor = ComputeColor(Difference->bAnyConflics, false);
-			FLinearColor LocalColor = ComputeColor(Difference->bAnyConflics, Difference->bLocalDifferences);
-			FLinearColor TextColor = ComputeColor(Difference->bAnyConflics, Difference->bLocalDifferences || Difference->bRemoteDifferences);
+			FLinearColor RemoteColor = ComputeColor(InDifference->bAnyConflics, InDifference->bRemoteDifferences);
+			FLinearColor BaseColor = ComputeColor(InDifference->bAnyConflics, false);
+			FLinearColor LocalColor = ComputeColor(InDifference->bAnyConflics, InDifference->bLocalDifferences);
+			FLinearColor TextColor = ComputeColor(InDifference->bAnyConflics, InDifference->bLocalDifferences || InDifference->bRemoteDifferences);
 
 			return SNew(SHorizontalBox)
 				+ SHorizontalBox::Slot()
 				[
 					SNew(STextBlock)
 					.ColorAndOpacity( TextColor )
-					.Text(FText::FromString(Difference->GraphName.GetPlainNameString()))
+					.Text(FText::FromString(InDifference->GraphName.GetPlainNameString()))
 				]
-				+ DiffViewUtils::Box(Difference->bExistsInRemote, RemoteColor)
-				+ DiffViewUtils::Box(Difference->bExistsInBase, BaseColor)
-				+ DiffViewUtils::Box(Difference->bExistsInLocal, LocalColor);
+				+ DiffViewUtils::Box(InDifference->bExistsInRemote, RemoteColor)
+				+ DiffViewUtils::Box(InDifference->bExistsInBase, BaseColor)
+				+ DiffViewUtils::Box(InDifference->bExistsInLocal, LocalColor);
 		};
 
-		const auto FocusGraph = [](FOnMergeNodeSelected SelectionCallback, SMergeGraphView* Parent, FName GraphName)
+		const auto FocusGraph = [](FOnMergeNodeSelected InSelectionCallback, SMergeGraphView* Parent, FName GraphName)
 		{
-			SelectionCallback.ExecuteIfBound();
+			InSelectionCallback.ExecuteIfBound();
 			Parent->FocusGraph( GraphName );
 		};
 
diff --git a/Engine/Source/Developer/Merge/Private/SMergeTreeView.cpp b/Engine/Source/Developer/Merge/Private/SMergeTreeView.cpp
--- a/Engine/Source/Developer/Merge/Private/SMergeTreeView.cpp
+++ b/Engine/Source/Developer/Merge/Private/SMergeTreeView.cpp
@@ -204,9 +204,9 @@ void SMergeTreeView::Construct(const FArguments InArgs
 			.ColorAndOpacity(Entry.bConflicted ? DiffViewUtils::Conflicting() : DiffViewUtils::Differs());
 	};
 
-	const auto FocusSCSDifferenceEntry = [](FSCSMergeEntry Entry, SMergeTreeView* Parent, FOnMergeNodeSelected SelectionCallback)
+	const auto FocusSCSDifferenceEntry = [](FSCSMergeEntry Entry, SMergeTreeView* Parent, FOnMergeNodeSelected InSelectionCallback)
 	{
-		SelectionCallback.ExecuteIfBound();
+		InSelectionCallback.ExecuteIfBound();
 		Parent->HighlightDifference(Entry.Identifier, Entry.PropertyIdentifier);
 	};
 	
@@ -230,10 +230,10 @@ void SMergeTreeView::Construct(const FArguments InArgs
 
 	DifferingProperties.Entries.Sort(SortTreePredicate);
 
-	const auto ForwardSelection = [](FOnMergeNodeSelected SelectionCallback)
+	const auto ForwardSelection = [](FOnMergeNodeSelected InSelectionCallback)
 	{
 		// This allows the owning control to focus the correct tab (or do whatever else it likes):
-		SelectionCallback.ExecuteIfBound();
+		InSelectionCallback.ExecuteIfBound();
 	};
 
 	const bool bHasDiffferences = Children.Num() != 0;
diff --git a/Engine/Source/Editor/DetailCustomizations/Private/SlateBrushCustomization.cpp b/Engine/Source/Editor/DetailCustomizations/Private/SlateBrushCustomization.cpp
--- a/Engine/Source/Editor/DetailCustomizations/Private/SlateBrushCustomization.cpp
+++ b/Engine/Source/Editor/DetailCustomizations/Private/SlateBrushCustomization.cpp
@@ -1175,8 +1175,8 @@ TSharedRef<IPropertyTypeCustomization> FSlateBrushStructCustomization::MakeInsta
 	return MakeShareable( new FSlateBrushStructCustomization(bIncludePreview) );
 }
 
-FSlateBrushStructCustomization::FSlateBrushStructCustomization(bool bIncludePreview)
-	: bIncludePreview(bIncludePreview)
+FSlateBrushStructCustomization::FSlateBrushStructCustomization(bool bInIncludePreview)
+	: bIncludePreview(bInIncludePreview)
 {
 }
 
diff --git a/Engine/Source/Editor/PhAT/Private/PhATSharedData.cpp b/Engine/Source/Editor/PhAT/Private/PhATSharedData.cpp
--- a/Engine/Source/Editor/PhAT/Private/PhATSharedData.cpp
+++ b/Engine/Source/Editor/PhAT/Private/PhATSharedData.cpp
@@ -1164,18 +1164,16 @@ void FPhATSharedData::DeleteBody(int32 DelBodyIndex, bool bRefreshComponent)
 
 		UBodySetup * ParentBody = ParentBodyIndex != INDEX_NONE ? PhysicsAsset->BodySetup[ParentBodyIndex] : NULL;
 
-		for (int32 i = 0; i < Constraints.Num(); ++i)
+		for (const int32 ConstraintIndex : Constraints)
 		{
-			int32 ConstraintIndex = Constraints[i];
 			UPhysicsConstraintTemplate * Constraint = PhysicsAsset->ConstraintSetup[ConstraintIndex];
 			Constraint->Modify();
 
 			if (ParentBody)
 			{
 				//for all constraints that contain a nearest child of this body, create a copy of the constraint between the child and parent
-				for (int32 i = 0; i < NearestBodiesBelow.Num(); ++i)
+				for (const int32 BodyBelowIndex : NearestBodiesBelow)
 				{
-					int32 BodyBelowIndex = NearestBodiesBelow[i];
 					UBodySetup * BodyBelow = PhysicsAsset->BodySetup[BodyBelowIndex];
 
 					if (Constraint->DefaultInstance.ConstraintBone1 == BodyBelow->BoneName)
diff --git a/Engine/Source/Editor/SceneOutliner/Private/ActorTreeItem.cpp b/Engine/Source/Editor/SceneOutliner/Private/ActorTreeItem.cpp
--- a/Engine/Source/Editor/SceneOutliner/Private/ActorTreeItem.cpp
+++ b/Engine/Source/Editor/SceneOutliner/Private/ActorTreeItem.cpp
@@ -39,9 +39,9 @@ FDragValidationInfo FActorDropTarget::ValidateDrop(FDragDropPayload& DraggedObje
 	bool bDraggedOntoAttachmentParent = true;
 
 	const auto& DragActors = DraggedObjects.Actors.GetValue();
-	for (const auto& Actor : DragActors)
+	for (const auto& DragActorPtr : DragActors)
 	{
-		AActor* DragActor = Actor.Get();
+		AActor* DragActor = DragActorPtr.Get();
 		if (DragActor)
 		{
 			if (bCanAttach)
diff --git a/Engine/Source/Editor/SceneOutliner/Private/SSceneOutliner.cpp b/Engine/Source/Editor/SceneOutliner/Private/SSceneOutliner.cpp
--- a/Engine/Source/Editor/SceneOutliner/Private/SSceneOutliner.cpp
+++ b/Engine/Source/Editor/SceneOutliner/Private/SSceneOutliner.cpp
@@ -1244,8 +1244,8 @@ namespace SceneOutliner
 
 			const bool bMoveToRootValid = FFolderDropTarget(FName()).ValidateDrop(DraggedObjects, *SharedData->RepresentingWorld).IsValid();
 
-			MiniSceneOutlinerInitOptions.Filters->AddFilterPredicate(FWorldFilterPredicate::CreateStatic([](const UWorld*, bool bMoveToRootValid){
-				return bMoveToRootValid;
+			MiniSceneOutlinerInitOptions.Filters->AddFilterPredicate(FWorldFilterPredicate::CreateStatic([](const UWorld*, bool bInMoveToRootValid){
+				return bInMoveToRootValid;
 			}, bMoveToRootValid), EDefaultFilterBehaviour::Pass);
 		}
 
diff --git a/Engine/Source/Editor/SceneOutliner/Private/SceneOutlinerActorInfoColumn.cpp b/Engine/Source/Editor/SceneOutliner/Private/SceneOutlinerActorInfoColumn.cpp
--- a/Engine/Source/Editor/SceneOutliner/Private/SceneOutlinerActorInfoColumn.cpp
+++ b/Engine/Source/Editor/SceneOutliner/Private/SceneOutlinerActorInfoColumn.cpp
@@ -197,8 +197,8 @@ const TSharedRef< SWidget > FActorInfoColumn::ConstructRowWidget( FTreeItemRef T
 			SNew(SBorder)
 			.BorderImage(FEditorStyle::GetBrush("NoBorder"))
 			.ColorAndOpacity_Static([](TWeakPtr<const STableRow<FTreeItemPtr>> WeakRow)->FLinearColor{
-				auto Row = WeakRow.Pin();
-				return Row.IsValid() && Row->IsSelected() ? FLinearColor::Black : FLinearColor::White;
+				auto TableRow = WeakRow.Pin();
+				return TableRow.IsValid() && TableRow->IsSelected() ? FLinearColor::Black : FLinearColor::White;
 			}, TWeakPtr<const STableRow<FTreeItemPtr>>(StaticCastSharedRef<const STableRow<FTreeItemPtr>>(Row.AsShared())))
 			[
 				Hyperlink.ToSharedRef()
@@ -252,10 +252,12 @@ TSharedPtr<SWidget> FActorInfoColumn::ConstructClassHyperlink( ITreeItem& TreeIt
 
 void FActorInfoColumn::PopulateSearchStrings( const ITreeItem& Item, TArray< FString >& OutSearchStrings ) const
 {
-	FString String = Item.Get(FGetInfo(CurrentMode));
-	if (String.Len())
 	{
-		OutSearchStrings.Add(String);
+		FString String = Item.Get(FGetInfo(CurrentMode));
+		if (String.Len())
+		{
+			OutSearchStrings.Add(String);
+		}
 	}
 
 	// We always add the class
diff --git a/Engine/Source/Editor/SceneOutliner/Private/WorldTreeItem.cpp b/Engine/Source/Editor/SceneOutliner/Private/WorldTreeItem.cpp
--- a/Engine/Source/Editor/SceneOutliner/Private/WorldTreeItem.cpp
+++ b/Engine/Source/Editor/SceneOutliner/Private/WorldTreeItem.cpp
@@ -89,17 +89,17 @@ void FWorldTreeItem::CreateFolder(TWeakPtr<SSceneOutliner> WeakOutliner)
 	}
 }
 
-FDragValidationInfo FWorldTreeItem::ValidateDrop(FDragDropPayload& DraggedObjects, UWorld& World) const
+FDragValidationInfo FWorldTreeItem::ValidateDrop(FDragDropPayload& DraggedObjects, UWorld& InWorld) const
 {
 	// Dropping on the world means 'moving to the root' in folder terms
 	FFolderDropTarget Target(NAME_None);
-	return Target.ValidateDrop(DraggedObjects, World);
+	return Target.ValidateDrop(DraggedObjects, InWorld);
 }
 
-void FWorldTreeItem::OnDrop(FDragDropPayload& DraggedObjects, UWorld& World, const FDragValidationInfo& ValidationInfo, TSharedRef<SWidget> DroppedOnWidget)
+void FWorldTreeItem::OnDrop(FDragDropPayload& DraggedObjects, UWorld& InWorld, const FDragValidationInfo& ValidationInfo, TSharedRef<SWidget> DroppedOnWidget)
 {
 	FFolderDropTarget Target(NAME_None);
-	return Target.OnDrop(DraggedObjects, World, ValidationInfo, DroppedOnWidget);
+	return Target.OnDrop(DraggedObjects, InWorld, ValidationInfo, DroppedOnWidget);
 }
 
 void FWorldTreeItem::OpenWorldSettings() const
diff --git a/Engine/Source/Editor/TranslationEditor/Private/TranslationDataManager.cpp b/Engine/Source/Editor/TranslationEditor/Private/TranslationDataManager.cpp
--- a/Engine/Source/Editor/TranslationEditor/Private/TranslationDataManager.cpp
+++ b/Engine/Source/Editor/TranslationEditor/Private/TranslationDataManager.cpp
@@ -564,10 +564,10 @@ void FTranslationDataManager::PreviewAllTranslationsInEditor()
 
 	FString ConfigFilePath = ConfigDirectory / "Localization" / "Regenerate" + FPaths::GetBaseFilename(ManifestFilePath) + ".ini";
 
-	FJsonInternationalizationArchiveSerializer ArchiveSerializer;
-	FJsonInternationalizationManifestSerializer ManifestSerializer;
+	FJsonInternationalizationArchiveSerializer LocalizationArchiveSerializer;
+	FJsonInternationalizationManifestSerializer LocalizationManifestSerializer;
 
-	FTextLocalizationManager::Get().RegenerateResources(ConfigFilePath, ArchiveSerializer, ManifestSerializer);
+	FTextLocalizationManager::Get().RegenerateResources(ConfigFilePath, LocalizationArchiveSerializer, LocalizationManifestSerializer);
 }
 
 void FTranslationDataManager::PopulateSearchResultsUsingFilter(const FString& SearchFilter)
