Details customizations re-enabled for the diff (and merge) tool. Now relying on the list of displayed properties from the details view when performing a CDO diff. This allows us to display changes to reference member variables.

diff --git a/Engine/Source/Developer/Merge/Private/SMergeDetailsView.cpp b/Engine/Source/Developer/Merge/Private/SMergeDetailsView.cpp
--- a/Engine/Source/Developer/Merge/Private/SMergeDetailsView.cpp
+++ b/Engine/Source/Developer/Merge/Private/SMergeDetailsView.cpp
@@ -21,21 +21,40 @@ void SMergeDetailsView::Construct(const FArguments InArgs
 	const UObject* BaseCDO = DiffUtils::GetCDO(InData.BlueprintBase);
 	const UObject* LocalCDO = DiffUtils::GetCDO(InData.BlueprintLocal);
 
+	// EMergeParticipant::Remote
+	{
+		DetailsViews.Add(
+			FDetailsDiff(RemoteCDO, FDetailsDiff::FOnDisplayedPropertiesChanged() )
+		);
+	}
+	// EMergeParticipant::Base
+	{
+		DetailsViews.Add(
+			FDetailsDiff(BaseCDO, FDetailsDiff::FOnDisplayedPropertiesChanged())
+		);
+	}
+	// EMergeParticipant::Local
+	{
+		DetailsViews.Add(
+			FDetailsDiff(LocalCDO, FDetailsDiff::FOnDisplayedPropertiesChanged())
+		);
+	}
+	
 	TArray<FSingleObjectDiffEntry> RemoteDifferences;
-	DiffUtils::CompareUnrelatedObjects(BaseCDO, RemoteCDO, RemoteDifferences);
+	GetBaseDetails().DiffAgainst(GetRemoteDetails(), RemoteDifferences);
 	TArray<FSingleObjectDiffEntry> LocalDifferences;
-	DiffUtils::CompareUnrelatedObjects(BaseCDO, LocalCDO, LocalDifferences);
+	GetBaseDetails().DiffAgainst(GetLocalDetails(), LocalDifferences);
 
-	const auto GetPropertyNames = []( const TArray<FSingleObjectDiffEntry>& FromDifferences, TArray<FPropertySoftPath>& OutNames )
+	const auto GetPropertyNames = [](const TArray<FSingleObjectDiffEntry>& FromDifferences, TArray<FPropertySoftPath>& OutNames)
 	{
-		for( const auto& Entry : FromDifferences )
+		for (const auto& Entry : FromDifferences)
 		{
 			OutNames.Push(Entry.Identifier);
 		}
 	};
 
 	TArray<FPropertySoftPath> RemoteDifferingProperties;
-	GetPropertyNames( RemoteDifferences, RemoteDifferingProperties );
+	GetPropertyNames(RemoteDifferences, RemoteDifferingProperties);
 	TArray<FPropertySoftPath> LocalDifferingProperties;
 	GetPropertyNames(LocalDifferences, LocalDifferingProperties);
 
@@ -44,46 +63,6 @@ void SMergeDetailsView::Construct(const FArguments InArgs
 	FPropertySoftPathSet BaseDifferingPropertiesSet(RemoteDifferingPropertiesSet.Union(LocalDifferingPropertiesSet));
 	MergeConflicts = RemoteDifferingPropertiesSet.Intersect(LocalDifferingPropertiesSet).Array();
 
-	/*	
-		DifferingProperties is an ordered list of all differing properties (properties added, removed or changed) in remote or local.
-		Strictly speaking it's impossible to guarantee that we'll traverse remote and local differences in the same order (for instance
-		because property layout could somehow change between revisions) but in practice the following works:
-
-			1. Iterate properties in base, add any properties that differ in remote or local to DifferingProperties.
-			2. Iterate properties in remote, add any new properties to DifferingProperties
-			3. Iterate properties in local, add any new properties to DifferingProperties, if they have not already been added
-	*/
-	const auto AddPropertiesOrdered = [](FPropertySoftPath const& Property, const FPropertySoftPathSet& InDifferingProperties, TArray<FPropertySoftPath>& ResultingProperties)
-	{
-		// contains check here is O(n), so we're needlessly n^2:
-		if (!ResultingProperties.Contains(Property))
-		{
-			if (InDifferingProperties.Contains(Property))
-			{
-				ResultingProperties.Add(Property);
-			}
-		}
-	};
-
-	// EMergeParticipant::Remote
-	{
-		DetailsViews.Add(
-			FDetailsDiff(RemoteCDO, DiffUtils::ResolveAll( RemoteCDO, RemoteDifferingProperties), FDetailsDiff::FOnDisplayedPropertiesChanged() )
-		);
-	}
-	// EMergeParticipant::Base
-	{
-		DetailsViews.Add(
-			FDetailsDiff(BaseCDO, DiffUtils::ResolveAll( BaseCDO, BaseDifferingPropertiesSet.Array()), FDetailsDiff::FOnDisplayedPropertiesChanged())
-		);
-	}
-	// EMergeParticipant::Local
-	{
-		DetailsViews.Add(
-			FDetailsDiff(LocalCDO, DiffUtils::ResolveAll( LocalCDO, LocalDifferingProperties), FDetailsDiff::FOnDisplayedPropertiesChanged())
-		);
-	}
-
 	TArray<FPropertySoftPath> RemoteVisibleProperties = GetRemoteDetails().GetDisplayedProperties();
 	TArray<FPropertySoftPath> BaseVisibleProperties = GetBaseDetails().GetDisplayedProperties();
 	TArray<FPropertySoftPath> LocalVisibleProperties = GetLocalDetails().GetDisplayedProperties();
@@ -153,6 +132,27 @@ void SMergeDetailsView::Construct(const FArguments InArgs
 		}
 	};
 
+	/*
+		DifferingProperties is an ordered list of all differing properties (properties added, removed or changed) in remote or local.
+		Strictly speaking it's impossible to guarantee that we'll traverse remote and local differences in the same order (for instance
+		because property layout could somehow change between revisions) but in practice the following works:
+
+		1. Iterate properties in base, add any properties that differ in remote or local to DifferingProperties.
+		2. Iterate properties in remote, add any new properties to DifferingProperties
+		3. Iterate properties in local, add any new properties to DifferingProperties, if they have not already been added
+	*/
+	const auto AddPropertiesOrdered = [](FPropertySoftPath const& Property, const FPropertySoftPathSet& InDifferingProperties, TArray<FPropertySoftPath>& ResultingProperties)
+	{
+		// contains check here is O(n), so we're needlessly n^2:
+		if (!ResultingProperties.Contains(Property))
+		{
+			if (InDifferingProperties.Contains(Property))
+			{
+				ResultingProperties.Add(Property);
+			}
+		}
+	};
+
 	TArray< FDiffPair > OrderedDifferences;
 	while( !bDoneRemote || !bDoneBase || !bDoneLocal )
 	{
diff --git a/Engine/Source/Editor/Kismet/Private/DetailsDiff.cpp b/Engine/Source/Editor/Kismet/Private/DetailsDiff.cpp
--- a/Engine/Source/Editor/Kismet/Private/DetailsDiff.cpp
+++ b/Engine/Source/Editor/Kismet/Private/DetailsDiff.cpp
@@ -6,9 +6,8 @@
 #include "PropertyEditorModule.h"
 //#include "SBlueprintDiff.h"
 
-FDetailsDiff::FDetailsDiff(const UObject* InObject, const TArray< FPropertyPath >& InDifferingProperties, FOnDisplayedPropertiesChanged InOnDisplayedPropertiesChanged )
+FDetailsDiff::FDetailsDiff(const UObject* InObject, FOnDisplayedPropertiesChanged InOnDisplayedPropertiesChanged )
 	: OnDisplayedPropertiesChanged( InOnDisplayedPropertiesChanged )
-	, DifferingProperties( InDifferingProperties )
 	, DisplayedObject( InObject )
 	, DetailsView()
 {
@@ -17,7 +16,6 @@ FDetailsDiff::FDetailsDiff(const UObject* InObject, const TArray< FPropertyPath
 
 	FPropertyEditorModule& EditModule = FModuleManager::Get().GetModuleChecked<FPropertyEditorModule>("PropertyEditor");
 	DetailsView = EditModule.CreateDetailView(DetailsViewArgs);
-	DetailsView->SetDisableCustomDetailLayouts(true);
 	DetailsView->SetIsPropertyEditingEnabledDelegate(FIsPropertyEditingEnabled::CreateStatic([]{return false; }));
 	// Forcing all advanced properties to be displayed for now, the logic to show changes made to advance properties
 	// conditionally is fragile and low priority for now:
@@ -26,7 +24,6 @@ FDetailsDiff::FDetailsDiff(const UObject* InObject, const TArray< FPropertyPath
 	// level would stress the type system a little:
 	DetailsView->SetObject(const_cast<UObject*>(InObject));
 	DetailsView->SetOnDisplayedPropertiesChanged( ::FOnDisplayedPropertiesChanged::CreateRaw(this, &FDetailsDiff::HandlePropertiesChanged) );
-	DetailsView->UpdatePropertiesWhitelist( TSet<FPropertyPath>(InDifferingProperties) );
 
 	DifferingProperties = DetailsView->GetPropertiesInOrderDisplayed();
 }
@@ -66,3 +63,63 @@ TArray<FPropertySoftPath> FDetailsDiff::GetDisplayedProperties() const
 	}
 	return Ret;
 }
+
+void FDetailsDiff::DiffAgainst(const FDetailsDiff& Newer, TArray< FSingleObjectDiffEntry > &OutDifferences) const
+{
+	TSharedPtr< class IDetailsView > OldDetailsView = DetailsView;
+	TSharedPtr< class IDetailsView > NewDetailsView = Newer.DetailsView;
+
+	const auto& OldSelectedObjects = OldDetailsView->GetSelectedObjects();
+	const auto& NewSelectedObjects = NewDetailsView->GetSelectedObjects();
+
+	check(OldSelectedObjects.Num() == 1);
+
+	auto OldProperties = OldDetailsView->GetPropertiesInOrderDisplayed();
+	auto NewProperties = NewDetailsView->GetPropertiesInOrderDisplayed();
+
+	TSet<FPropertySoftPath> OldPropertiesSet;
+	TSet<FPropertySoftPath> NewPropertiesSet;
+
+	auto ToSet = [](const TArray<FPropertyPath>& Array, TSet<FPropertySoftPath>& OutSet)
+	{
+		for (const auto& Entry : Array)
+		{
+			OutSet.Add(FPropertySoftPath(Entry));
+		}
+	};
+
+	ToSet(OldProperties, OldPropertiesSet);
+	ToSet(NewProperties, NewPropertiesSet);
+
+	// detect removed properties:
+	auto RemovedProperties = OldPropertiesSet.Difference(NewPropertiesSet);
+	for (const auto& RemovedProperty : RemovedProperties)
+	{
+		// @todo: (doc) label these as removed, rather than added to a
+		FSingleObjectDiffEntry Entry(RemovedProperty, EPropertyDiffType::PropertyAddedToA);
+		OutDifferences.Push(Entry);
+	}
+
+	// detect added properties:
+	auto AddededProperties = NewPropertiesSet.Difference(OldPropertiesSet);
+	for (const auto& AddedProperty : AddededProperties)
+	{
+		FSingleObjectDiffEntry Entry(AddedProperty, EPropertyDiffType::PropertyAddedToB);
+		OutDifferences.Push(Entry);
+	}
+
+	// check for changed properties
+	auto CommonProperties = NewPropertiesSet.Intersect(OldPropertiesSet);
+	for (const auto& CommonProperty : CommonProperties)
+	{
+		// get value, diff:
+		check(NewSelectedObjects.Num() == 1);
+		auto OldPoperty = CommonProperty.Resolve(OldSelectedObjects[0].Get());
+		auto NewProperty = CommonProperty.Resolve(NewSelectedObjects[0].Get());
+
+		if (!DiffUtils::Identical(OldPoperty, NewProperty))
+		{
+			OutDifferences.Push(FSingleObjectDiffEntry(CommonProperty, EPropertyDiffType::PropertyValueChanged));
+		}
+	}
+}
diff --git a/Engine/Source/Editor/Kismet/Private/DiffUtils.cpp b/Engine/Source/Editor/Kismet/Private/DiffUtils.cpp
--- a/Engine/Source/Editor/Kismet/Private/DiffUtils.cpp
+++ b/Engine/Source/Editor/Kismet/Private/DiffUtils.cpp
@@ -31,16 +31,28 @@ static FPropertySoftPathSet GetPropertyNameSet(const UObject* ForObj)
 FResolvedProperty FPropertySoftPath::Resolve(const UObject* Object) const
 {
 	// dig into the object, finding nested objects, etc:
+	const void* CurrentBlock = Object;
+	const UStruct* NextClass = Object->GetClass();
+	const void* NextBlock = CurrentBlock;
 	const UProperty* Property = nullptr;
+
 	for( int32 i = 0; i < PropertyChain.Num(); ++i )
 	{
-		const UProperty* NextProperty = ::Resolve(Object->GetClass(), PropertyChain[i]);
+		CurrentBlock = NextBlock;
+		const UProperty* NextProperty = ::Resolve(NextClass, PropertyChain[i]);
 		if( NextProperty )
 		{
 			Property = NextProperty;
 			if (const UObjectProperty* ObjectProperty = Cast<UObjectProperty>(Property))
 			{
-				Object = ObjectProperty->GetObjectPropertyValue(Property->ContainerPtrToValuePtr<UObject*>(Object));
+				const UObject* NextObject = ObjectProperty->GetObjectPropertyValue(Property->ContainerPtrToValuePtr<UObject*>(CurrentBlock));
+				NextBlock = NextObject;
+				NextClass = NextObject ? NextObject->GetClass() : nullptr;
+			}
+			else if (const UStructProperty* StructProperty = Cast<UStructProperty>(Property))
+			{
+				NextBlock = StructProperty->ContainerPtrToValuePtr<void>(CurrentBlock);
+				NextClass = StructProperty->Struct;
 			}
 			else
 			{
@@ -53,7 +65,7 @@ FResolvedProperty FPropertySoftPath::Resolve(const UObject* Object) const
 		}
 	}
 
-	return FResolvedProperty(Object, Property);
+	return FResolvedProperty(CurrentBlock, Property);
 }
 
 FPropertyPath FPropertySoftPath::ResolvePath(const UObject* Object) const
diff --git a/Engine/Source/Editor/Kismet/Private/SBlueprintDiff.cpp b/Engine/Source/Editor/Kismet/Private/SBlueprintDiff.cpp
--- a/Engine/Source/Editor/Kismet/Private/SBlueprintDiff.cpp
+++ b/Engine/Source/Editor/Kismet/Private/SBlueprintDiff.cpp
@@ -141,7 +141,6 @@ class FCDODiffControl	: public TSharedFromThis<FCDODiffControl>
 public:
 	FCDODiffControl( const UObject* InOldCDO
 					, const UObject* InNewCDO
-					, TArray< FSingleObjectDiffEntry > const& InDifferences
 					, TArray< TSharedPtr<FBlueprintDifferenceTreeEntry> >& OutTreeEntries
 					, TArray< TSharedPtr<FBlueprintDifferenceTreeEntry> >& OutRealDifferences
 					, FOnCDODiffControlChanged SelectionCallback );
@@ -164,13 +163,17 @@ private:
 FCDODiffControl::FCDODiffControl( 
 		const UObject* InOldCDO
 		, const UObject* InNewCDO
-		, TArray< FSingleObjectDiffEntry > const& InDifferences
 		, TArray< TSharedPtr<FBlueprintDifferenceTreeEntry> >& OutTreeEntries
 		, TArray< TSharedPtr<FBlueprintDifferenceTreeEntry> >& OutRealDifferences
 		, FOnCDODiffControlChanged SelectionCallback )
-	: OldDetails(InOldCDO, DiffUtils::ResolveAll( InOldCDO, InDifferences), FDetailsDiff::FOnDisplayedPropertiesChanged() )
-	, NewDetails(InNewCDO, DiffUtils::ResolveAll( InNewCDO, InDifferences), FDetailsDiff::FOnDisplayedPropertiesChanged() )
-{// OrderedProperties will contain differences in the order they are displayed:
+	: OldDetails(InOldCDO, FDetailsDiff::FOnDisplayedPropertiesChanged() )
+	, NewDetails(InNewCDO, FDetailsDiff::FOnDisplayedPropertiesChanged() )
+{
+	// @todo: (doc): rename InDifferences
+	TArray< FSingleObjectDiffEntry > DifferingProperties;
+	OldDetails.DiffAgainst(NewDetails, DifferingProperties);
+
+	// OrderedProperties will contain differences in the order they are displayed:
 	TArray< const FSingleObjectDiffEntry* > OrderedProperties;
 
 	// create differing properties list based on what is displayed by the old properties..
@@ -198,29 +201,48 @@ FCDODiffControl::FCDODiffControl(
 		if (IterOld != OldProperties.Num() && IterNew != NewProperties.Num()
 			&& OldProperties[IterOld] == NewProperties[IterNew])
 		{
-			if (const FSingleObjectDiffEntry* Differing = FindDiffering(InDifferences, OldProperties[IterOld]))
+			if (const FSingleObjectDiffEntry* Differing = FindDiffering(DifferingProperties, OldProperties[IterOld]))
 			{
 				OrderedProperties.Push(Differing);
 			}
 			++IterOld;
 			++IterNew;
 		}
-		else if (IterOld != OldProperties.Num())
+		else
 		{
-			if (const FSingleObjectDiffEntry* Differing = FindDiffering(InDifferences, OldProperties[IterOld]))
+			if (IterOld != OldProperties.Num())
 			{
-				OrderedProperties.Push(Differing);
+				bool bFoundDifference = false;
+				if (const FSingleObjectDiffEntry* Differing = FindDiffering(DifferingProperties, OldProperties[IterOld]))
+				{
+					bFoundDifference = true;
+					OrderedProperties.Push(Differing);
+					++IterOld;
+				}
+				else if (IterNew != NewProperties.Num())
+				{
+					if (const FSingleObjectDiffEntry* Differing = FindDiffering(DifferingProperties, NewProperties[IterNew]))
+					{
+						bFoundDifference = true;
+						OrderedProperties.Push(Differing);
+						++IterNew;
+					}
+				}
+				if (!bFoundDifference)
+				{
+					++IterOld;
+					++IterNew;
+				}
 			}
-			++IterOld;
-		}
-		else
-		{
-			check(IterNew != NewProperties.Num());
-			if (const FSingleObjectDiffEntry* Differing = FindDiffering(InDifferences, NewProperties[IterNew]))
+			else
 			{
-				OrderedProperties.Push(Differing);
+				check(IterNew != NewProperties.Num());
+				if (const FSingleObjectDiffEntry* Differing = FindDiffering(DifferingProperties, NewProperties[IterNew]))
+				{
+					OrderedProperties.Push(Differing);
+				}
+				++IterNew;
 			}
-			++IterNew;
 		}
 	}
 
@@ -1274,10 +1296,8 @@ SBlueprintDiff::FDiffControl SBlueprintDiff::GenerateDefaultsPanel()
 {
 	const UObject* A = DiffUtils::GetCDO(PanelOld.Blueprint);
 	const UObject* B = DiffUtils::GetCDO(PanelNew.Blueprint);
-	TArray<FSingleObjectDiffEntry> DifferingProperties;
-	DiffUtils::CompareUnrelatedObjects( A, B, DifferingProperties);
 
-	auto NewDiffControl = TSharedPtr<FCDODiffControl>(new FCDODiffControl(A, B, DifferingProperties, MasterDifferencesList, RealDifferences, FOnCDODiffControlChanged::CreateRaw(this, &SBlueprintDiff::SetCurrentMode, FBlueprintEditorApplicationModes::BlueprintDefaultsMode)));
+	auto NewDiffControl = TSharedPtr<FCDODiffControl>(new FCDODiffControl(A, B, MasterDifferencesList, RealDifferences, FOnCDODiffControlChanged::CreateRaw(this, &SBlueprintDiff::SetCurrentMode, FBlueprintEditorApplicationModes::BlueprintDefaultsMode)));
 
 	SBlueprintDiff::FDiffControl Ret;
 	//Splitter for left and right blueprint. Current convention is for the local (probably newer?) blueprint to be on the right:
diff --git a/Engine/Source/Editor/Kismet/Public/DetailsDiff.h b/Engine/Source/Editor/Kismet/Public/DetailsDiff.h
--- a/Engine/Source/Editor/Kismet/Public/DetailsDiff.h
+++ b/Engine/Source/Editor/Kismet/Public/DetailsDiff.h
@@ -9,13 +9,15 @@ class KISMET_API FDetailsDiff
 public:
 	DECLARE_DELEGATE( FOnDisplayedPropertiesChanged );
 
-	FDetailsDiff( const UObject* InObject, const TArray< FPropertyPath >& InDifferingProperties, FOnDisplayedPropertiesChanged InOnDisplayedPropertiesChanged );
+	FDetailsDiff( const UObject* InObject, FOnDisplayedPropertiesChanged InOnDisplayedPropertiesChanged );
 	~FDetailsDiff();
 
 	void HighlightProperty( const FPropertySoftPath& PropertyName );
 	TSharedRef< SWidget > DetailsWidget();
 	TArray<FPropertySoftPath> GetDisplayedProperties() const;
 
+	void DiffAgainst(const FDetailsDiff& Newer, TArray< FSingleObjectDiffEntry > &OutDifferences) const;
+
 private:
 	void HandlePropertiesChanged();
 
diff --git a/Engine/Source/Editor/Kismet/Public/DiffUtils.h b/Engine/Source/Editor/Kismet/Public/DiffUtils.h
--- a/Engine/Source/Editor/Kismet/Public/DiffUtils.h
+++ b/Engine/Source/Editor/Kismet/Public/DiffUtils.h
@@ -2,17 +2,16 @@
 #pragma once
 
 #include "PropertyPath.h"
-#include "IAssetTypeActions.h"
 
 struct FResolvedProperty
 {
 	explicit FResolvedProperty()
-	: Object(nullptr)
-	, Property(nullptr)
+		: Object(nullptr)
+		, Property(nullptr)
 	{
 	}
 
-	FResolvedProperty(const UObject* InObject, const UProperty* InProperty)
+	FResolvedProperty(const void* InObject, const UProperty* InProperty)
 		: Object(InObject)
 		, Property(InProperty)
 	{
@@ -25,7 +24,7 @@ struct FResolvedProperty
 
 	inline bool operator!=(const FResolvedProperty& RHS) const { return !(*this == RHS); }
 
-	const UObject* Object;
+	const void* Object;
 	const UProperty* Property;
 };
 
@@ -226,6 +225,8 @@ namespace DiffTreeView
 	KISMET_API bool HasPrevDifference(TSharedRef< STreeView<TSharedPtr< FBlueprintDifferenceTreeEntry > > > TreeView, const TArray< TSharedPtr<class FBlueprintDifferenceTreeEntry> >& Differences);
 }
 
+struct FRevisionInfo;
+
 namespace DiffViewUtils
 {
 	KISMET_API FLinearColor LookupColor( bool bDiffers, bool bConflicts = false );
@@ -236,7 +237,7 @@ namespace DiffViewUtils
 
 	KISMET_API FText PropertyDiffMessage(FSingleObjectDiffEntry Difference, FText ObjectName);
 	KISMET_API FText SCSDiffMessage(const FSCSDiffEntry& Difference, FText ObjectName);
-	KISMET_API FText GetPanelLabel( const UBlueprint* Blueprint, const FRevisionInfo& Revision, FText Label );
+	KISMET_API FText GetPanelLabel(const UBlueprint* Blueprint, FRevisionInfo const& Revision, FText Label);
 
 	KISMET_API SHorizontalBox::FSlot& Box(bool bIsPresent, FLinearColor Color);
 }
