I am not sure that the merge will work, but at least it won't crash

diff --git a/Engine/Source/Developer/Merge/Private/SBlueprintMerge.cpp b/Engine/Source/Developer/Merge/Private/SBlueprintMerge.cpp
--- a/Engine/Source/Developer/Merge/Private/SBlueprintMerge.cpp
+++ b/Engine/Source/Developer/Merge/Private/SBlueprintMerge.cpp
@@ -264,7 +264,7 @@ void SBlueprintMerge::Construct(const FArguments InArgs, const FBlueprintMergeDa
 
 UBlueprint* SBlueprintMerge::GetTargetBlueprint()
 {
-	return Data.OwningEditor.Pin()->GetBlueprintObj();
+	return Data.OwningEditor.IsValid() ? Data.OwningEditor.Pin()->GetBlueprintObj() : nullptr;
 }
 
 void SBlueprintMerge::NextDiff()
@@ -534,9 +534,13 @@ void SBlueprintMerge::ResolveMerge(UBlueprint* ResultantBlueprint)
 	// in the case where we have replaced/reloaded the blueprint 
 	// ("TargetBlueprint" is most likely, and should be, forcefully closing in 
 	// this scenario... we need to open the result to take its place)
-	if (ResultantBlueprint != GetTargetBlueprint())
+	UBlueprint* TargetBlueprint = GetTargetBlueprint();
+	if (ResultantBlueprint != TargetBlueprint)
 	{
-		FAssetEditorManager::Get().CloseAllEditorsForAsset(GetTargetBlueprint());
+		if (TargetBlueprint)
+		{
+			FAssetEditorManager::Get().CloseAllEditorsForAsset(TargetBlueprint);
+		}
 		FAssetEditorManager::Get().OpenEditorForAsset(ResultantBlueprint);
 	}
 	
diff --git a/Engine/Source/Editor/UnrealEd/Private/Kismet2/Kismet2.cpp b/Engine/Source/Editor/UnrealEd/Private/Kismet2/Kismet2.cpp
--- a/Engine/Source/Editor/UnrealEd/Private/Kismet2/Kismet2.cpp
+++ b/Engine/Source/Editor/UnrealEd/Private/Kismet2/Kismet2.cpp
@@ -684,11 +684,15 @@ UBlueprint* FKismetEditorUtilities::ReloadBlueprint(UBlueprint* StaleBlueprint)
 	check(StaleBlueprint->IsAsset());
 	FSoftObjectPath BlueprintAssetRef(StaleBlueprint);
 
+	// Need to close the editor now, it won't work once it's been garbage collected during the reload
+	FAssetEditorManager::Get().CloseAllEditorsForAsset(StaleBlueprint);
+
 	FBlueprintUnloader Unloader(StaleBlueprint);
 	Unloader.UnloadBlueprint(/*bResetPackage =*/true);
 
 	UBlueprint* ReloadedBlueprint = Cast<UBlueprint>(StaticLoadObject(UBlueprint::StaticClass(), /*Outer =*/nullptr, *BlueprintAssetRef.ToString()));
 
+	// This will generally not fix most references as the UnloadBlueprint causes a GC
 	Unloader.ReplaceStaleRefs(ReloadedBlueprint);
 	return ReloadedBlueprint;
 }
diff --git a/Engine/Source/Editor/UnrealEd/Public/Kismet2/KismetEditorUtilities.h b/Engine/Source/Editor/UnrealEd/Public/Kismet2/KismetEditorUtilities.h
--- a/Engine/Source/Editor/UnrealEd/Public/Kismet2/KismetEditorUtilities.h
+++ b/Engine/Source/Editor/UnrealEd/Public/Kismet2/KismetEditorUtilities.h
@@ -107,6 +107,7 @@ public:
 	/** 
 	 * Unloads the supplied Blueprint (marking it pending-kill, and removing it 
 	 * from its outer package). Then proceeds to reload from disk.
+	 * This will generally not transfer object references over as there is a garbage collect in the middle.
 	 *
 	 * @param  TargetBlueprint	The Blueprint you want to unload and replace.
 	 * @return The freshly loaded Blueprint (replacing the, now invalid, input).
