Added logging to Blueprint merge operation

diff --git a/Engine/Source/Developer/Merge/Private/SBlueprintMerge.cpp b/Engine/Source/Developer/Merge/Private/SBlueprintMerge.cpp
--- a/Engine/Source/Developer/Merge/Private/SBlueprintMerge.cpp
+++ b/Engine/Source/Developer/Merge/Private/SBlueprintMerge.cpp
@@ -516,6 +516,9 @@ void SBlueprintMerge::OnAcceptLocal()
 void SBlueprintMerge::ResolveMerge(UBlueprint* ResultantBlueprint)
 {
 	UPackage* Package = ResultantBlueprint->GetOutermost();
+	
+	UE_LOG(LogUObjectGlobals, Log, TEXT("Merging Blueprint: %s"), *Package->GetName());
+
 	TArray<UPackage*> PackagesToSave;
 	PackagesToSave.Add(Package);
 
