PR #6098: Fix editor crash when accepting local revision during merge (Contributed by [REDACTED_PERSON])

diff --git a/Engine/Source/Developer/Merge/Private/MergeUtils.cpp b/Engine/Source/Developer/Merge/Private/MergeUtils.cpp
--- a/Engine/Source/Developer/Merge/Private/MergeUtils.cpp
+++ b/Engine/Source/Developer/Merge/Private/MergeUtils.cpp
@@ -38,40 +38,13 @@ UObject const* FMergeToolUtils::LoadRevision(const FString& AssetName, const ISo
 {
 	FSlateNotificationManager& NotificationManager = FSlateNotificationManager::Get();
 
+	const UObject* AssetRevision = nullptr;
+
 	// Get the head revision of this package from source control
 	FString TempFileName;
 	if (DesiredRevision.Get(TempFileName))
 	{
-		// Try and load that package
-		UPackage* TempPackage = LoadPackage(NULL, *TempFileName, LOAD_ForDiff|LOAD_DisableCompileOnLoad);
-		if (TempPackage != NULL)
-		{
-			// Grab the old asset from that old package
-			UObject* OldObject = FindObject<UObject>(TempPackage, *AssetName);
-			if (OldObject)
-			{
-				return OldObject;
-			}
-			else
-			{
-				NotificationManager.AddNotification(
-					FText::Format(
-						LOCTEXT("MergedFailedToFindObject", "Aborted Load of {0} because we could not find an object named {1}" )
-						, FText::FromString(TempFileName)
-						, FText::FromString(AssetName) 
-					)
-				);
-			}
-		}
-		else
-		{
-			NotificationManager.AddNotification(
-				FText::Format(
-					LOCTEXT("MergedFailedToLoadPackage", "Aborted Load of {0} because we could not load the package")
-					, FText::FromString(TempFileName)
-				)
-			);
-		}
+		AssetRevision = LoadAssetFromPackage(TempFileName, AssetName);
 	}
 	else
 	{
@@ -82,7 +55,8 @@ UObject const* FMergeToolUtils::LoadRevision(const FString& AssetName, const ISo
 			)
 		);
 	}
-	return NULL;
+
+	return AssetRevision;
 }
 
 //------------------------------------------------------------------------------
@@ -136,4 +110,45 @@ UObject const* FMergeToolUtils::LoadRevision(const UObject* AssetObject, const F
 	return AssetRevision;
 }
 
+//------------------------------------------------------------------------------
+UObject const* FMergeToolUtils::LoadAssetFromPackage(const FString& PackageFileName, const FString& AssetName)
+{
+	FSlateNotificationManager& NotificationManager = FSlateNotificationManager::Get();
+	
+	const UObject* Object = nullptr;
+
+	// Try and load that package
+	UPackage* TempPackage = LoadPackage(NULL, *PackageFileName, LOAD_ForDiff | LOAD_DisableCompileOnLoad);
+	if (TempPackage != NULL)
+	{
+		// Grab the old asset from that old package
+		UObject* FoundObject = FindObject<UObject>(TempPackage, *AssetName);
+		if (FoundObject)
+		{
+			Object = FoundObject;
+		}
+		else
+		{
+			NotificationManager.AddNotification(
+				FText::Format(
+					LOCTEXT("MergedFailedToFindObject", "Aborted Load of {0} because we could not find an object named {1}")
+					, FText::FromString(PackageFileName)
+					, FText::FromString(AssetName)
+				)
+			);
+		}
+	}
+	else
+	{
+		NotificationManager.AddNotification(
+			FText::Format(
+				LOCTEXT("MergedFailedToLoadPackage", "Aborted Load of {0} because we could not load the package")
+				, FText::FromString(PackageFileName)
+			)
+		);
+	}
+
+	return Object;
+}
+
 #undef LOCTEXT_NAMESPACE
diff --git a/Engine/Source/Developer/Merge/Private/MergeUtils.h b/Engine/Source/Developer/Merge/Private/MergeUtils.h
--- a/Engine/Source/Developer/Merge/Private/MergeUtils.h
+++ b/Engine/Source/Developer/Merge/Private/MergeUtils.h
@@ -60,4 +60,13 @@ struct FMergeToolUtils
 	 * @return The requested version of the specified asset object, could be null if the supplied revision failed to load.
 	 */
 	static UObject const* LoadRevision(const UObject* AssetObject, const FRevisionInfo& DesiredRevision);
+
+	/**
+	 * Attempts to load the specified asset from the provided package.
+	 * 
+	 * @param  PackageFileName  The file name of the package where to look for the asset.
+	 * @param  AssetName        The name of the asset you want loaded from the package.
+	 * @return The asset object you wanted loaded, could be null if the asset wasn't found in the package.
+	 */
+	static UObject const* LoadAssetFromPackage(const FString& PackageFileName, const FString& AssetName);
 };
diff --git a/Engine/Source/Developer/Merge/Private/SBlueprintMerge.cpp b/Engine/Source/Developer/Merge/Private/SBlueprintMerge.cpp
--- a/Engine/Source/Developer/Merge/Private/SBlueprintMerge.cpp
+++ b/Engine/Source/Developer/Merge/Private/SBlueprintMerge.cpp
@@ -473,38 +473,18 @@ void SBlueprintMerge::OnAcceptLocal()
 		// we use the emptiness of LocalBackupPath as a sentinel value (to 
 		// determine how we should replace the TargetBlueprint)... if 
 		// Data.BlueprintLocal is malformed (and not ok to copy), then 
-		// LocalBackupPath should be valid (and we need to copy the file on disk 
-		// instead)
+		// LocalBackupPath should be valid (and we need to replace with
+		// the backup instead)
 		if (LocalBackupPath.IsEmpty())
 		{
 			NewBlueprint = FKismetEditorUtilities::ReplaceBlueprint(TargetBlueprint, Data.BlueprintLocal);
 		}
 		else
 		{
-			UPackage* TargetPackage = TargetBlueprint->GetOutermost();
+			const FString BlueprintName = TargetBlueprint->GetName();
+			const UBlueprint* LocalBackupBlueprint = Cast<const UBlueprint>(FMergeToolUtils::LoadAssetFromPackage(LocalBackupPath, BlueprintName));
 
-			FString PackageFilename;
-			if (FPackageName::DoesPackageExist(TargetPackage->GetName(), /*Guid =*/nullptr, &PackageFilename))
-			{
-				FString SrcFilename = LocalBackupPath;
-				auto OverwriteBlueprintFile = [TargetBlueprint, &PackageFilename, &SrcFilename](UBlueprint* /*UnloadedBP*/)
-				{
-					if (IFileManager::Get().Copy(*PackageFilename, *SrcFilename) != COPY_OK)
-					{
-						const FText TargetName = FText::FromName(TargetBlueprint->GetFName());
-						const FText ErrorMessage = FText::Format(LOCTEXT("FailedMergeLocalCopy", "Failed to overwrite {0} target file."), TargetName);
-						FSlateNotificationManager::Get().AddNotification(FNotificationInfo(ErrorMessage));
-					}
-				};
-				// we have to wait until the package is unloaded (and the file is 
-				// unlocked) before we can overwrite the file; so we defer the
-				// copy until then by adding it to the OnBlueprintUnloaded delegate
-				FKismetEditorUtilities::FOnBlueprintUnloaded::FDelegate OnPackageUnloaded = FKismetEditorUtilities::FOnBlueprintUnloaded::FDelegate::CreateLambda(OverwriteBlueprintFile);
-					
-				FDelegateHandle OnPackageUnloadedDelegateHandle = FKismetEditorUtilities::OnBlueprintUnloaded.Add(OnPackageUnloaded);
-				NewBlueprint = FKismetEditorUtilities::ReloadBlueprint(TargetBlueprint);
-				FKismetEditorUtilities::OnBlueprintUnloaded.Remove(OnPackageUnloadedDelegateHandle);
-			}
+			NewBlueprint = FKismetEditorUtilities::ReplaceBlueprint(TargetBlueprint, LocalBackupBlueprint);
 		}
 	}
 
diff --git a/Engine/Source/Editor/UnrealEd/Private/Kismet2/Kismet2.cpp b/Engine/Source/Editor/UnrealEd/Private/Kismet2/Kismet2.cpp
--- a/Engine/Source/Editor/UnrealEd/Private/Kismet2/Kismet2.cpp
+++ b/Engine/Source/Editor/UnrealEd/Private/Kismet2/Kismet2.cpp
@@ -713,24 +713,6 @@ UK2Node_Event* FKismetEditorUtilities::AddDefaultEventNode(UBlueprint* InBluepri
 	return NewEventNode;
 }
 
-UBlueprint* FKismetEditorUtilities::ReloadBlueprint(UBlueprint* StaleBlueprint)
-{
-	check(StaleBlueprint->IsAsset());
-	FSoftObjectPath BlueprintAssetRef(StaleBlueprint);
-
-	// Need to close the editor now, it won't work once it's been garbage collected during the reload
-	GEditor->GetEditorSubsystem<UAssetEditorSubsystem>()->CloseAllEditorsForAsset(StaleBlueprint);
-
-	FBlueprintUnloader Unloader(StaleBlueprint);
-	Unloader.UnloadBlueprint(/*bResetPackage =*/true);
-
-	UBlueprint* ReloadedBlueprint = Cast<UBlueprint>(StaticLoadObject(UBlueprint::StaticClass(), /*Outer =*/nullptr, *BlueprintAssetRef.ToString()));
-
-	// This will generally not fix most references as the UnloadBlueprint causes a GC
-	Unloader.ReplaceStaleRefs(ReloadedBlueprint);
-	return ReloadedBlueprint;
-}
-
 UBlueprint* FKismetEditorUtilities::ReplaceBlueprint(UBlueprint* Target, UBlueprint const* ReplacementArchetype)
 {
 	if (Target == ReplacementArchetype)
diff --git a/Engine/Source/Editor/UnrealEd/Public/Kismet2/KismetEditorUtilities.h b/Engine/Source/Editor/UnrealEd/Public/Kismet2/KismetEditorUtilities.h
--- a/Engine/Source/Editor/UnrealEd/Public/Kismet2/KismetEditorUtilities.h
+++ b/Engine/Source/Editor/UnrealEd/Public/Kismet2/KismetEditorUtilities.h
@@ -106,21 +106,11 @@ public:
 
 	/** 
 	 * Event that's broadcast anytime a blueprint is unloaded, and becomes 
-	 * invalid (with calls to ReloadBlueprint(), for example).
+	 * invalid (with calls to ReplaceBlueprint(), for example).
 	 */
 	DECLARE_MULTICAST_DELEGATE_OneParam(FOnBlueprintUnloaded, UBlueprint*);
 	static FOnBlueprintUnloaded OnBlueprintUnloaded;
 
-	/** 
-	 * Unloads the supplied Blueprint (marking it pending-kill, and removing it 
-	 * from its outer package). Then proceeds to reload from disk.
-	 * This will generally not transfer object references over as there is a garbage collect in the middle.
-	 *
-	 * @param  TargetBlueprint	The Blueprint you want to unload and replace.
-	 * @return The freshly loaded Blueprint (replacing the, now invalid, input).
-	 */
-	static UBlueprint* ReloadBlueprint(UBlueprint* TargetBlueprint);
-
 	/** 
 	 * Unloads the specified Blueprint (marking it pending-kill, and removing it 
 	 * from its outer package). Then proceeds to replace all references with a
