Tab manager filter can now block tab spawn

diff --git a/Engine/Plugins/Animation/LiveLinkCurveDebugUI/Source/LiveLinkCurveDebugUI/Private/LiveLinkCurveDebugUIModule.cpp b/Engine/Plugins/Animation/LiveLinkCurveDebugUI/Source/LiveLinkCurveDebugUI/Private/LiveLinkCurveDebugUIModule.cpp
--- a/Engine/Plugins/Animation/LiveLinkCurveDebugUI/Source/LiveLinkCurveDebugUI/Private/LiveLinkCurveDebugUIModule.cpp
+++ b/Engine/Plugins/Animation/LiveLinkCurveDebugUI/Source/LiveLinkCurveDebugUI/Private/LiveLinkCurveDebugUIModule.cpp
@@ -184,7 +184,7 @@ void FLiveLinkCurveDebugUIModule::ShutdownModule()
 void FLiveLinkCurveDebugUIModule::DisplayThroughTab()
 {
 	check(bHasRegisteredTabSpawners);
-	FGlobalTabmanager::Get()->InvokeTab(FTabId("LiveLinkCurveDebug"));
+	FGlobalTabmanager::Get()->TryInvokeTab(FTabId("LiveLinkCurveDebug"));
 }
 
 TSharedPtr<SLiveLinkCurveDebugUI> FLiveLinkCurveDebugUIModule::CreateDebugWidget()
diff --git a/Engine/Plugins/Developer/ActorLayerUtilities/Source/ActorLayerUtilitiesEditor/Private/ActorLayerPropertyTypeCustomization.cpp b/Engine/Plugins/Developer/ActorLayerUtilities/Source/ActorLayerUtilitiesEditor/Private/ActorLayerPropertyTypeCustomization.cpp
--- a/Engine/Plugins/Developer/ActorLayerUtilities/Source/ActorLayerUtilitiesEditor/Private/ActorLayerPropertyTypeCustomization.cpp
+++ b/Engine/Plugins/Developer/ActorLayerUtilities/Source/ActorLayerUtilitiesEditor/Private/ActorLayerPropertyTypeCustomization.cpp
@@ -202,7 +202,7 @@ void FActorLayerPropertyTypeCustomization::AssignLayer(FName InNewLayer)
 void FActorLayerPropertyTypeCustomization::OpenLayerBrowser()
 {
 	FLevelEditorModule& LevelEditorModule = FModuleManager::GetModuleChecked<FLevelEditorModule>("LevelEditor");
-	LevelEditorModule.GetLevelEditorTabManager()->InvokeTab(FTabId("LevelEditorLayerBrowser"));
+	LevelEditorModule.GetLevelEditorTabManager()->TryInvokeTab(FTabId("LevelEditorLayerBrowser"));
 }
 
 FReply FActorLayerPropertyTypeCustomization::OnDrop(TSharedPtr<FDragDropOperation> InDragDrop)
diff --git a/Engine/Plugins/Developer/Concert/MultiUserClient/Source/MultiUserClient/Private/MultiUserClientModule.cpp b/Engine/Plugins/Developer/Concert/MultiUserClient/Source/MultiUserClient/Private/MultiUserClientModule.cpp
--- a/Engine/Plugins/Developer/Concert/MultiUserClient/Source/MultiUserClient/Private/MultiUserClientModule.cpp
+++ b/Engine/Plugins/Developer/Concert/MultiUserClient/Source/MultiUserClient/Private/MultiUserClientModule.cpp
@@ -624,7 +624,7 @@ public:
 	 */
 	virtual void OpenBrowser() override
 	{
-		FGlobalTabmanager::Get()->InvokeTab(FTabId(ConcertBrowserTabName));
+		FGlobalTabmanager::Get()->TryInvokeTab(FTabId(ConcertBrowserTabName));
 	}
 
 	/**
@@ -972,7 +972,7 @@ private:
 		}
 
 		// The auto-connect is not configured, clicking the button will open (or flash) the multi-users browser.
-		FGlobalTabmanager::Get()->InvokeTab(FTabId(ConcertBrowserTabName));
+		FGlobalTabmanager::Get()->TryInvokeTab(FTabId(ConcertBrowserTabName));
 	}
 
 	/**
diff --git a/Engine/Plugins/Editor/AssetManagerEditor/Source/AssetManagerEditor/Private/AssetManagerEditorModule.cpp b/Engine/Plugins/Editor/AssetManagerEditor/Source/AssetManagerEditor/Private/AssetManagerEditorModule.cpp
--- a/Engine/Plugins/Editor/AssetManagerEditor/Source/AssetManagerEditor/Private/AssetManagerEditorModule.cpp
+++ b/Engine/Plugins/Editor/AssetManagerEditor/Source/AssetManagerEditor/Private/AssetManagerEditorModule.cpp
@@ -647,7 +647,7 @@ TSharedRef<SDockTab> FAssetManagerEditorModule::SpawnSizeMapTab(const FSpawnTabA
 
 void FAssetManagerEditorModule::OpenAssetAuditUI(TArray<FAssetData> SelectedAssets)
 {
-	FGlobalTabmanager::Get()->InvokeTab(AssetAuditTabName);
+	FGlobalTabmanager::Get()->TryInvokeTab(AssetAuditTabName);
 
 	if (AssetAuditUI.IsValid())
 	{
@@ -657,7 +657,7 @@ void FAssetManagerEditorModule::OpenAssetAuditUI(TArray<FAssetData> SelectedAsse
 
 void FAssetManagerEditorModule::OpenAssetAuditUI(TArray<FAssetIdentifier> SelectedIdentifiers)
 {
-	FGlobalTabmanager::Get()->InvokeTab(AssetAuditTabName);
+	FGlobalTabmanager::Get()->TryInvokeTab(AssetAuditTabName);
 
 	if (AssetAuditUI.IsValid())
 	{
@@ -667,7 +667,7 @@ void FAssetManagerEditorModule::OpenAssetAuditUI(TArray<FAssetIdentifier> Select
 
 void FAssetManagerEditorModule::OpenAssetAuditUI(TArray<FName> SelectedPackages)
 {
-	FGlobalTabmanager::Get()->InvokeTab(AssetAuditTabName);
+	FGlobalTabmanager::Get()->TryInvokeTab(AssetAuditTabName);
 
 	if (AssetAuditUI.IsValid())
 	{
@@ -679,9 +679,11 @@ void FAssetManagerEditorModule::OpenReferenceViewerUI(const TArray<FAssetIdentif
 {
 	if (SelectedIdentifiers.Num() > 0)
 	{
-		TSharedRef<SDockTab> NewTab = FGlobalTabmanager::Get()->InvokeTab(ReferenceViewerTabName);
-		TSharedRef<SReferenceViewer> ReferenceViewer = StaticCastSharedRef<SReferenceViewer>(NewTab->GetContent());
-		ReferenceViewer->SetGraphRootIdentifiers(SelectedIdentifiers, ReferenceViewerParams);
+		if (TSharedPtr<SDockTab> NewTab = FGlobalTabmanager::Get()->TryInvokeTab(ReferenceViewerTabName))
+		{
+			TSharedRef<SReferenceViewer> ReferenceViewer = StaticCastSharedRef<SReferenceViewer>(NewTab->GetContent());
+			ReferenceViewer->SetGraphRootIdentifiers(SelectedIdentifiers, ReferenceViewerParams);
+		}
 	}
 }
 
@@ -711,9 +713,11 @@ void FAssetManagerEditorModule::OpenSizeMapUI(TArray<FAssetIdentifier> SelectedI
 {
 	if (SelectedIdentifiers.Num() > 0)
 	{
-		TSharedRef<SDockTab> NewTab = FGlobalTabmanager::Get()->InvokeTab(SizeMapTabName);
-		TSharedRef<SSizeMap> SizeMap = StaticCastSharedRef<SSizeMap>(NewTab->GetContent());
-		SizeMap->SetRootAssetIdentifiers(SelectedIdentifiers);
+		if (TSharedPtr<SDockTab> NewTab = FGlobalTabmanager::Get()->TryInvokeTab(SizeMapTabName))
+		{
+			TSharedRef<SSizeMap> SizeMap = StaticCastSharedRef<SSizeMap>(NewTab->GetContent());
+			SizeMap->SetRootAssetIdentifiers(SelectedIdentifiers);
+		}
 	}
 }
 
@@ -746,10 +750,12 @@ void FAssetManagerEditorModule::OpenShaderCookStatistics(TArray<FName> SelectedP
 	static const FName LevelEditorStatsViewerTab("LevelEditorStatsViewer");
 	FLevelEditorModule& LevelEditorModule = FModuleManager::GetModuleChecked<FLevelEditorModule>(LevelEditorModuleName);
 	TSharedPtr<FTabManager> TabManager = LevelEditorModule.GetLevelEditorTabManager();
-	TSharedRef<SDockTab> Tab = TabManager->InvokeTab(LevelEditorStatsViewerTab);
-	TSharedRef<SWidget> Content = Tab->GetContent();
-	IStatsViewer* StatsView = (IStatsViewer*)&*Content;
-	StatsView->SwitchAndFilterPage(EStatsPage::ShaderCookerStats, CommonPath, FString("Path"));
+	if (TSharedPtr<SDockTab> Tab = TabManager->TryInvokeTab(LevelEditorStatsViewerTab))
+	{
+		TSharedRef<SWidget> Content = Tab->GetContent();
+		IStatsViewer* StatsView = (IStatsViewer*)&*Content;
+		StatsView->SwitchAndFilterPage(EStatsPage::ShaderCookerStats, CommonPath, FString("Path"));
+	}
 }
 
 void FAssetManagerEditorModule::GetAssetDataInPaths(const TArray<FString>& Paths, TArray<FAssetData>& OutAssetData)
diff --git a/Engine/Plugins/Editor/PluginBrowser/Source/PluginBrowser/Private/PluginBrowserModule.cpp b/Engine/Plugins/Editor/PluginBrowser/Source/PluginBrowser/Private/PluginBrowserModule.cpp
--- a/Engine/Plugins/Editor/PluginBrowser/Source/PluginBrowser/Private/PluginBrowserModule.cpp
+++ b/Engine/Plugins/Editor/PluginBrowser/Source/PluginBrowser/Private/PluginBrowserModule.cpp
@@ -179,7 +179,7 @@ void FPluginBrowserModule::OnMainFrameLoaded(TSharedPtr<SWindow> InRootWindow, b
 
 void FPluginBrowserModule::OnNewPluginsPopupSettingsClicked()
 {
-	FGlobalTabmanager::Get()->InvokeTab(PluginsEditorTabName);
+	FGlobalTabmanager::Get()->TryInvokeTab(PluginsEditorTabName);
 	NewPluginsNotification.Pin()->ExpireAndFadeout();
 }
 
diff --git a/Engine/Plugins/Editor/PluginBrowser/Source/PluginBrowser/Private/SPluginBrowser.cpp b/Engine/Plugins/Editor/PluginBrowser/Source/PluginBrowser/Private/SPluginBrowser.cpp
--- a/Engine/Plugins/Editor/PluginBrowser/Source/PluginBrowser/Private/SPluginBrowser.cpp
+++ b/Engine/Plugins/Editor/PluginBrowser/Source/PluginBrowser/Private/SPluginBrowser.cpp
@@ -390,7 +390,7 @@ void SPluginBrowser::BreadcrumbTrail_OnCrumbClicked( const TSharedPtr<FPluginCat
 
 FReply SPluginBrowser::HandleNewPluginButtonClicked() const
 {
-	FGlobalTabmanager::Get()->InvokeTab( FPluginBrowserModule::PluginCreatorTabName );
+	FGlobalTabmanager::Get()->TryInvokeTab( FPluginBrowserModule::PluginCreatorTabName );
 
 	return FReply::Handled();
 }
diff --git a/Engine/Plugins/Editor/PluginBrowser/Templates/Advanced/Source/PLUGIN_NAME/Private/PLUGIN_NAME.cpp b/Engine/Plugins/Editor/PluginBrowser/Templates/Advanced/Source/PLUGIN_NAME/Private/PLUGIN_NAME.cpp
--- a/Engine/Plugins/Editor/PluginBrowser/Templates/Advanced/Source/PLUGIN_NAME/Private/PLUGIN_NAME.cpp
+++ b/Engine/Plugins/Editor/PluginBrowser/Templates/Advanced/Source/PLUGIN_NAME/Private/PLUGIN_NAME.cpp
@@ -76,7 +76,7 @@ TSharedRef<SDockTab> FPLUGIN_NAMEModule::OnSpawnPluginTab(const FSpawnTabArgs& S
 
 void FPLUGIN_NAMEModule::PluginButtonClicked()
 {
-	FGlobalTabmanager::Get()->InvokeTab(PLUGIN_NAMETabName);
+	FGlobalTabmanager::Get()->TryInvokeTab(PLUGIN_NAMETabName);
 }
 
 void FPLUGIN_NAMEModule::RegisterMenus()
diff --git a/Engine/Plugins/Enterprise/DataprepEditor/Source/DataprepEditor/Private/DataprepDetailsTab.cpp b/Engine/Plugins/Enterprise/DataprepEditor/Source/DataprepEditor/Private/DataprepDetailsTab.cpp
--- a/Engine/Plugins/Enterprise/DataprepEditor/Source/DataprepEditor/Private/DataprepDetailsTab.cpp
+++ b/Engine/Plugins/Enterprise/DataprepEditor/Source/DataprepEditor/Private/DataprepDetailsTab.cpp
@@ -108,7 +108,7 @@ void FDataprepEditor::TryInvokingDetailsTab(bool bFlash)
 				if (!DetailsTabPtr.IsValid())
 				{
 					// Show the details panel if it doesn't exist.
-					TabManager->InvokeTab(FDataprepEditor::DetailsTabId);
+					TabManager->TryInvokeTab(FDataprepEditor::DetailsTabId);
 				}
 			}
 		}
diff --git a/Engine/Plugins/Enterprise/DataprepEditor/Source/DataprepEditor/Private/DataprepPipelineEditor.cpp b/Engine/Plugins/Enterprise/DataprepEditor/Source/DataprepEditor/Private/DataprepPipelineEditor.cpp
--- a/Engine/Plugins/Enterprise/DataprepEditor/Source/DataprepEditor/Private/DataprepPipelineEditor.cpp
+++ b/Engine/Plugins/Enterprise/DataprepEditor/Source/DataprepEditor/Private/DataprepPipelineEditor.cpp
@@ -677,7 +677,7 @@ void FDataprepEditor::OnCompile()
 			// DATAPREP_TODO: Create tab for compiling error
 			//if (!bEditorMarkedAsClosed && bForceMessageDisplay && GetCurrentMode() == FBlueprintEditorApplicationModes::StandardBlueprintEditorMode)
 			//{
-			//	TabManager->InvokeTab(FBlueprintEditorTabs::CompilerResultsID);
+			//	TabManager->TryInvokeTab(FBlueprintEditorTabs::CompilerResultsID);
 			//}
 		}
 
diff --git a/Engine/Plugins/Enterprise/StaticMeshEditorExtension/Source/StaticMeshEditorExtension/Private/UVTools/UVGenerationToolbar.cpp b/Engine/Plugins/Enterprise/StaticMeshEditorExtension/Source/StaticMeshEditorExtension/Private/UVTools/UVGenerationToolbar.cpp
--- a/Engine/Plugins/Enterprise/StaticMeshEditorExtension/Source/StaticMeshEditorExtension/Private/UVTools/UVGenerationToolbar.cpp
+++ b/Engine/Plugins/Enterprise/StaticMeshEditorExtension/Source/StaticMeshEditorExtension/Private/UVTools/UVGenerationToolbar.cpp
@@ -91,7 +91,7 @@ void FUVGenerationToolbar::CreateUVMenu(FMenuBuilder& ParentMenuBuilder, const T
 		{
 			return;
 		}
-		EditorInstance->GetAssociatedTabManager()->InvokeTab(UVGenerationToolTabId);
+		EditorInstance->GetAssociatedTabManager()->TryInvokeTab(UVGenerationToolTabId);
 	});
 	ParentMenuBuilder.AddMenuEntry(
 		LOCTEXT("GenerateUVs", "Generate UVs"),
diff --git a/Engine/Plugins/Enterprise/VariantManager/Source/VariantManager/Private/LevelVariantSetsEditorToolkit.cpp b/Engine/Plugins/Enterprise/VariantManager/Source/VariantManager/Private/LevelVariantSetsEditorToolkit.cpp
--- a/Engine/Plugins/Enterprise/VariantManager/Source/VariantManager/Private/LevelVariantSetsEditorToolkit.cpp
+++ b/Engine/Plugins/Enterprise/VariantManager/Source/VariantManager/Private/LevelVariantSetsEditorToolkit.cpp
@@ -86,15 +86,17 @@ void FLevelVariantSetsEditorToolkit::Initialize(const EToolkitMode::Type Mode, c
 	// tab ID will already be registered within EditorTabManager
 	if (EditorTabManager->FindExistingLiveTab(TabID).IsValid())
 	{
-		EditorTabManager->InvokeTab(TabID)->RequestCloseTab();
+		EditorTabManager->TryInvokeTab(TabID)->RequestCloseTab();
 	}
 
 	FAssetEditorToolkit::InitAssetEditor(Mode, InitToolkitHost, TEXT("VariantManagerApp"), StandaloneDefaultLayout, bCreateDefaultStandaloneMenu, bCreateDefaultToolbar, LevelVariantSets);
 
-	TSharedRef<SDockTab> Tab = EditorTabManager->InvokeTab(TabID);
-	Tab->SetContent(VariantManager->GetVariantManagerWidget().ToSharedRef());
-	Tab->SetOnTabClosed(SDockTab::FOnTabClosedCallback::CreateStatic(&Local::OnVariantManagerClosed, TWeakPtr<IAssetEditorInstance>(SharedThis(this))));
-
+	TSharedPtr<SDockTab> Tab = EditorTabManager->TryInvokeTab(TabID);
+	if (Tab.IsValid())
+	{
+		Tab->SetContent(VariantManager->GetVariantManagerWidget().ToSharedRef());
+		Tab->SetOnTabClosed(SDockTab::FOnTabClosedCallback::CreateStatic(&Local::OnVariantManagerClosed, TWeakPtr<IAssetEditorInstance>(SharedThis(this))));
+	}
 	CreatedTab = Tab;
 }
 
diff --git a/Engine/Plugins/FX/Niagara/Source/NiagaraEditor/Private/Toolkits/NiagaraSystemToolkit.cpp b/Engine/Plugins/FX/Niagara/Source/NiagaraEditor/Private/Toolkits/NiagaraSystemToolkit.cpp
--- a/Engine/Plugins/FX/Niagara/Source/NiagaraEditor/Private/Toolkits/NiagaraSystemToolkit.cpp
+++ b/Engine/Plugins/FX/Niagara/Source/NiagaraEditor/Private/Toolkits/NiagaraSystemToolkit.cpp
@@ -1474,7 +1474,7 @@ bool FNiagaraSystemToolkit::OnApplyEnabled() const
 
 void FNiagaraSystemToolkit::OnPinnedCurvesChanged()
 {
-	TabManager->InvokeTab(CurveEditorTabID);
+	TabManager->TryInvokeTab(CurveEditorTabID);
 }
 
 void FNiagaraSystemToolkit::RefreshParameters()
@@ -1514,7 +1514,7 @@ void FNiagaraSystemToolkit::OnSystemSelectionChanged()
 
 void FNiagaraSystemToolkit::OnViewModelRequestFocusTab(FName TabName)
 {
-	GetTabManager()->InvokeTab(TabName);
+	GetTabManager()->TryInvokeTab(TabName);
 }
 
 #undef LOCTEXT_NAMESPACE
diff --git a/Engine/Plugins/MovieScene/ActorSequence/Source/ActorSequenceEditor/Private/ActorSequenceComponentCustomization.cpp b/Engine/Plugins/MovieScene/ActorSequence/Source/ActorSequenceEditor/Private/ActorSequenceComponentCustomization.cpp
--- a/Engine/Plugins/MovieScene/ActorSequence/Source/ActorSequenceEditor/Private/ActorSequenceComponentCustomization.cpp
+++ b/Engine/Plugins/MovieScene/ActorSequence/Source/ActorSequenceEditor/Private/ActorSequenceComponentCustomization.cpp
@@ -153,33 +153,34 @@ FReply FActorSequenceComponentCustomization::InvokeSequencer()
 	TSharedPtr<FTabManager> TabManager = WeakTabManager.Pin();
 	if (TabManager.IsValid() && TabManager->HasTabSpawner(SequenceTabId))
 	{
-		TSharedRef<SDockTab> Tab = TabManager->InvokeTab(SequenceTabId);
-
+		if (TSharedPtr<SDockTab> Tab = TabManager->TryInvokeTab(SequenceTabId))
 		{
-			// Set up a delegate that forces a refresh of this panel when the tab is closed to ensure we see the inline widget
-			TWeakPtr<IPropertyUtilities> WeakUtilities = PropertyUtilities;
-			auto OnClosed = [WeakUtilities](TSharedRef<SDockTab>)
 			{
-				TSharedPtr<IPropertyUtilities> PinnedPropertyUtilities = WeakUtilities.Pin();
-				if (PinnedPropertyUtilities.IsValid())
+				// Set up a delegate that forces a refresh of this panel when the tab is closed to ensure we see the inline widget
+				TWeakPtr<IPropertyUtilities> WeakUtilities = PropertyUtilities;
+				auto OnClosed = [WeakUtilities](TSharedRef<SDockTab>)
 				{
-					PinnedPropertyUtilities->EnqueueDeferredAction(FSimpleDelegate::CreateSP(PinnedPropertyUtilities.ToSharedRef(), &IPropertyUtilities::ForceRefresh));
-				}
-			};
-
-			Tab->SetOnTabClosed(SDockTab::FOnTabClosedCallback::CreateLambda(OnClosed));
-		}
-
-		// Move our inline widget content to the tab (so that we keep the existing sequencer state)
-		if (InlineSequencer.IsValid())
-		{
-			Tab->SetContent(InlineSequencer->GetChildren()->GetChildAt(0));
-			InlineSequencer->SetContent(SNullWidget::NullWidget);
-			InlineSequencer->SetVisibility(EVisibility::Collapsed);
-		}
-		else 
-		{
-			StaticCastSharedRef<SActorSequenceEditorWidget>(Tab->GetContent())->AssignSequence(GetActorSequence());
+					TSharedPtr<IPropertyUtilities> PinnedPropertyUtilities = WeakUtilities.Pin();
+					if (PinnedPropertyUtilities.IsValid())
+					{
+						PinnedPropertyUtilities->EnqueueDeferredAction(FSimpleDelegate::CreateSP(PinnedPropertyUtilities.ToSharedRef(), &IPropertyUtilities::ForceRefresh));
+					}
+				};
+
+				Tab->SetOnTabClosed(SDockTab::FOnTabClosedCallback::CreateLambda(OnClosed));
+			}
+
+			// Move our inline widget content to the tab (so that we keep the existing sequencer state)
+			if (InlineSequencer.IsValid())
+			{
+				Tab->SetContent(InlineSequencer->GetChildren()->GetChildAt(0));
+				InlineSequencer->SetContent(SNullWidget::NullWidget);
+				InlineSequencer->SetVisibility(EVisibility::Collapsed);
+			}
+			else
+			{
+				StaticCastSharedRef<SActorSequenceEditorWidget>(Tab->GetContent())->AssignSequence(GetActorSequence());
+			}
 		}
 	}
 
diff --git a/Engine/Plugins/MovieScene/LevelSequenceEditor/Source/LevelSequenceEditor/Private/LevelSequenceEditorToolkit.cpp b/Engine/Plugins/MovieScene/LevelSequenceEditor/Source/LevelSequenceEditor/Private/LevelSequenceEditorToolkit.cpp
--- a/Engine/Plugins/MovieScene/LevelSequenceEditor/Source/LevelSequenceEditor/Private/LevelSequenceEditorToolkit.cpp
+++ b/Engine/Plugins/MovieScene/LevelSequenceEditor/Source/LevelSequenceEditor/Private/LevelSequenceEditorToolkit.cpp
@@ -212,8 +212,8 @@ void FLevelSequenceEditorToolkit::Initialize(const EToolkitMode::Type Mode, cons
 		TSharedPtr<FTabManager> LevelEditorTabManager = LevelEditorModule.GetLevelEditorTabManager();
 		if (LevelEditorTabManager->FindExistingLiveTab(FName("LevelEditorSceneOutliner")).IsValid())
 		{
-			LevelEditorTabManager->InvokeTab(FName("LevelEditorSceneOutliner"))->RequestCloseTab();
-			LevelEditorTabManager->InvokeTab(FName("LevelEditorSceneOutliner"));
+			LevelEditorTabManager->TryInvokeTab(FName("LevelEditorSceneOutliner"))->RequestCloseTab();
+			LevelEditorTabManager->TryInvokeTab(FName("LevelEditorSceneOutliner"));
 		}
 	}
 	
diff --git a/Engine/Plugins/MovieScene/MatineeToLevelSequence/Source/MatineeToLevelSequence/Private/MatineeToLevelSequenceModule.cpp b/Engine/Plugins/MovieScene/MatineeToLevelSequence/Source/MatineeToLevelSequence/Private/MatineeToLevelSequenceModule.cpp
--- a/Engine/Plugins/MovieScene/MatineeToLevelSequence/Source/MatineeToLevelSequence/Private/MatineeToLevelSequenceModule.cpp
+++ b/Engine/Plugins/MovieScene/MatineeToLevelSequence/Source/MatineeToLevelSequence/Private/MatineeToLevelSequenceModule.cpp
@@ -425,7 +425,7 @@ protected:
 
 			FNotificationInfo NotificationInfo(NotificationText);
 			NotificationInfo.ExpireDuration = 5.f;
-			NotificationInfo.Hyperlink = FSimpleDelegate::CreateStatic([](){ FGlobalTabmanager::Get()->InvokeTab(FName("OutputLog")); });
+			NotificationInfo.Hyperlink = FSimpleDelegate::CreateStatic([](){ FGlobalTabmanager::Get()->TryInvokeTab(FName("OutputLog")); });
 			NotificationInfo.HyperlinkText = LOCTEXT("ShowMessageLogHyperlink", "Show Output Log");
 			FSlateNotificationManager::Get().AddNotification(NotificationInfo);
 		}
diff --git a/Engine/Plugins/MovieScene/TemplateSequence/Source/TemplateSequenceEditor/Private/TemplateSequenceEditorToolkit.cpp b/Engine/Plugins/MovieScene/TemplateSequence/Source/TemplateSequenceEditor/Private/TemplateSequenceEditorToolkit.cpp
--- a/Engine/Plugins/MovieScene/TemplateSequence/Source/TemplateSequenceEditor/Private/TemplateSequenceEditorToolkit.cpp
+++ b/Engine/Plugins/MovieScene/TemplateSequence/Source/TemplateSequenceEditor/Private/TemplateSequenceEditorToolkit.cpp
@@ -127,8 +127,8 @@ void FTemplateSequenceEditorToolkit::Initialize(const EToolkitMode::Type Mode, c
 		TSharedPtr<FTabManager> LevelEditorTabManager = LevelEditorModule.GetLevelEditorTabManager();
 		if (LevelEditorTabManager->FindExistingLiveTab(FName("LevelEditorSceneOutliner")).IsValid())
 		{
-			LevelEditorTabManager->InvokeTab(FName("LevelEditorSceneOutliner"))->RequestCloseTab();
-			LevelEditorTabManager->InvokeTab(FName("LevelEditorSceneOutliner"));
+			LevelEditorTabManager->TryInvokeTab(FName("LevelEditorSceneOutliner"))->RequestCloseTab();
+			LevelEditorTabManager->TryInvokeTab(FName("LevelEditorSceneOutliner"));
 		}
 	}
 
diff --git a/Engine/Plugins/NetcodeUnitTest/NetcodeUnitTest/Source/NetcodeUnitTest/Private/UI/SLogWidget.cpp b/Engine/Plugins/NetcodeUnitTest/NetcodeUnitTest/Source/NetcodeUnitTest/Private/UI/SLogWidget.cpp
--- a/Engine/Plugins/NetcodeUnitTest/NetcodeUnitTest/Source/NetcodeUnitTest/Private/UI/SLogWidget.cpp
+++ b/Engine/Plugins/NetcodeUnitTest/NetcodeUnitTest/Source/NetcodeUnitTest/Private/UI/SLogWidget.cpp
@@ -1100,10 +1100,10 @@ void SLogWidget::AddLine(ELogType InLogType, TSharedRef<FString> LogLine, FSlate
 			// If the tab is not presently open, open it now
 			if (!CurTabInfo->bTabOpen && LogTabManager.IsValid())
 			{
-				LogTabManager->InvokeTab(CurTabInfo->TabIdName);
+				LogTabManager->TryInvokeTab(CurTabInfo->TabIdName);
 
 				// The new tab has stolen focus, now restore the old tabs focus
-				LogTabManager->InvokeTab(ActiveTab->TabIdName);
+				LogTabManager->TryInvokeTab(ActiveTab->TabIdName);
 
 				CurTabInfo->bTabOpen = true;
 			}
@@ -1145,7 +1145,7 @@ void SLogWidget::AddLine(ELogType InLogType, TSharedRef<FString> LogLine, FSlate
 	// If a focus change is required, perform it
 	if (FocusTab.IsValid() && LogTabManager.IsValid())
 	{
-		LogTabManager->InvokeTab(FocusTab->TabIdName);
+		LogTabManager->TryInvokeTab(FocusTab->TabIdName);
 	}
 }
 
diff --git a/Engine/Plugins/Runtime/GameplayAbilities/Source/GameplayAbilitiesEditor/Private/GameplayAbilitiesEditorModule.cpp b/Engine/Plugins/Runtime/GameplayAbilities/Source/GameplayAbilitiesEditor/Private/GameplayAbilitiesEditorModule.cpp
--- a/Engine/Plugins/Runtime/GameplayAbilities/Source/GameplayAbilitiesEditor/Private/GameplayAbilitiesEditorModule.cpp
+++ b/Engine/Plugins/Runtime/GameplayAbilities/Source/GameplayAbilitiesEditor/Private/GameplayAbilitiesEditorModule.cpp
@@ -179,7 +179,7 @@ void FGameplayAbilitiesEditorModule::StartupModule()
 	{
 		FLevelEditorModule& LevelEditorModule = FModuleManager::GetModuleChecked<FLevelEditorModule>( TEXT("LevelEditor") );
 		TSharedPtr<FTabManager> LevelEditorTabManager = LevelEditorModule.GetLevelEditorTabManager();
-		LevelEditorTabManager->InvokeTab(FName("GameplayCueApp"));
+		LevelEditorTabManager->TryInvokeTab(FName("GameplayCueApp"));
 	}
 #endif // WITH_HOT_RELOAD
 
@@ -198,7 +198,7 @@ void FGameplayAbilitiesEditorModule::HandleNotify_OpenAssetInEditor(FString Asse
 	{
 		FLevelEditorModule& LevelEditorModule = FModuleManager::GetModuleChecked<FLevelEditorModule>(TEXT("LevelEditor"));
 		TSharedPtr<FTabManager> LevelEditorTabManager = LevelEditorModule.GetLevelEditorTabManager();
-		LevelEditorTabManager->InvokeTab(FName("GameplayCueApp"));
+		LevelEditorTabManager->TryInvokeTab(FName("GameplayCueApp"));
 	}
 
 	//UE_LOG(LogTemp, Display, TEXT("HandleNotify_OpenAssetInEditor!!! %s %d"), *AssetName, AssetType);
@@ -215,7 +215,7 @@ void FGameplayAbilitiesEditorModule::HandleNotify_FindAssetInEditor(FString Asse
 	{
 		FLevelEditorModule& LevelEditorModule = FModuleManager::GetModuleChecked<FLevelEditorModule>(TEXT("LevelEditor"));
 		TSharedPtr<FTabManager> LevelEditorTabManager = LevelEditorModule.GetLevelEditorTabManager();
-		LevelEditorTabManager->InvokeTab(FName("GameplayCueApp"));
+		LevelEditorTabManager->TryInvokeTab(FName("GameplayCueApp"));
 	}
 
 	//UE_LOG(LogTemp, Display, TEXT("HandleNotify_FindAssetInEditor!!! %s %d"), *AssetName, AssetType);
diff --git a/Engine/Plugins/Runtime/Oculus/OculusVR/Source/OculusEditor/Private/OculusEditorModule.cpp b/Engine/Plugins/Runtime/Oculus/OculusVR/Source/OculusEditor/Private/OculusEditorModule.cpp
--- a/Engine/Plugins/Runtime/Oculus/OculusVR/Source/OculusEditor/Private/OculusEditorModule.cpp
+++ b/Engine/Plugins/Runtime/Oculus/OculusVR/Source/OculusEditor/Private/OculusEditorModule.cpp
@@ -176,7 +176,7 @@ FReply FOculusEditorModule::PluginClickFn(bool text)
 
 void FOculusEditorModule::PluginButtonClicked()
 {
-	FGlobalTabmanager::Get()->InvokeTab(OculusPerfTabName);
+	FGlobalTabmanager::Get()->TryInvokeTab(OculusPerfTabName);
 }
 
 void FOculusEditorModule::AddMenuExtension(FMenuBuilder& Builder)
@@ -206,13 +206,13 @@ TSharedRef<IDetailCustomization> FOculusHMDSettingsDetailsCustomization::MakeIns
 
 FReply FOculusHMDSettingsDetailsCustomization::PluginClickPerfFn(bool text)
 {
-	FGlobalTabmanager::Get()->InvokeTab(FOculusEditorModule::OculusPerfTabName);
+	FGlobalTabmanager::Get()->TryInvokeTab(FOculusEditorModule::OculusPerfTabName);
 	return FReply::Handled();
 }
 
 FReply FOculusHMDSettingsDetailsCustomization::PluginClickPlatFn(bool text)
 {
-	FGlobalTabmanager::Get()->InvokeTab(FOculusEditorModule::OculusPlatToolTabName);
+	FGlobalTabmanager::Get()->TryInvokeTab(FOculusEditorModule::OculusPlatToolTabName);
 	return FReply::Handled();
 }
 
diff --git a/Engine/Plugins/Runtime/Oculus/OculusVR/Source/OculusEditor/Private/OculusToolCommands.cpp b/Engine/Plugins/Runtime/Oculus/OculusVR/Source/OculusEditor/Private/OculusToolCommands.cpp
--- a/Engine/Plugins/Runtime/Oculus/OculusVR/Source/OculusEditor/Private/OculusToolCommands.cpp
+++ b/Engine/Plugins/Runtime/Oculus/OculusVR/Source/OculusEditor/Private/OculusToolCommands.cpp
@@ -12,7 +12,7 @@ void FOculusToolCommands::RegisterCommands()
 
 void FOculusToolCommands::ShowOculusTool()
 {
-	FGlobalTabmanager::Get()->InvokeTab(FOculusEditorModule::OculusPerfTabName);
+	FGlobalTabmanager::Get()->TryInvokeTab(FOculusEditorModule::OculusPerfTabName);
 }
 
 #undef LOCTEXT_NAMESPACE
diff --git a/Engine/Plugins/Runtime/Steam/SteamAudio/Source/SteamAudioEditor/Private/BakeIndirectWindow.cpp b/Engine/Plugins/Runtime/Steam/SteamAudio/Source/SteamAudioEditor/Private/BakeIndirectWindow.cpp
--- a/Engine/Plugins/Runtime/Steam/SteamAudio/Source/SteamAudioEditor/Private/BakeIndirectWindow.cpp
+++ b/Engine/Plugins/Runtime/Steam/SteamAudio/Source/SteamAudioEditor/Private/BakeIndirectWindow.cpp
@@ -196,7 +196,7 @@ namespace SteamAudio
 	 */
 	void FBakeIndirectWindow::Invoke()
 	{
-		FGlobalTabmanager::Get()->InvokeTab(BakeIndirectTabName);
+		FGlobalTabmanager::Get()->TryInvokeTab(BakeIndirectTabName);
 	}
 
 	/**
diff --git a/Engine/Plugins/VirtualProduction/Takes/Source/TakeRecorder/Private/Recorder/TakeRecorderBlueprintLibrary.cpp b/Engine/Plugins/VirtualProduction/Takes/Source/TakeRecorder/Private/Recorder/TakeRecorderBlueprintLibrary.cpp
--- a/Engine/Plugins/VirtualProduction/Takes/Source/TakeRecorder/Private/Recorder/TakeRecorderBlueprintLibrary.cpp
+++ b/Engine/Plugins/VirtualProduction/Takes/Source/TakeRecorder/Private/Recorder/TakeRecorderBlueprintLibrary.cpp
@@ -112,7 +112,7 @@ UTakeRecorderPanel* UTakeRecorderBlueprintLibrary::OpenTakeRecorderPanel()
 	if (!IsRunningGame())
 	{
 		FLevelEditorModule& LevelEditorModule = FModuleManager::LoadModuleChecked<FLevelEditorModule>("LevelEditor");
-		LevelEditorModule.GetLevelEditorTabManager()->InvokeTab(ITakeRecorderModule::TakeRecorderTabName);
+		LevelEditorModule.GetLevelEditorTabManager()->TryInvokeTab(ITakeRecorderModule::TakeRecorderTabName);
 	}
 	else
 	{
diff --git a/Engine/Plugins/VirtualProduction/Takes/Source/TakeRecorder/Private/TakePresetActions.cpp b/Engine/Plugins/VirtualProduction/Takes/Source/TakeRecorder/Private/TakePresetActions.cpp
--- a/Engine/Plugins/VirtualProduction/Takes/Source/TakeRecorder/Private/TakePresetActions.cpp
+++ b/Engine/Plugins/VirtualProduction/Takes/Source/TakeRecorder/Private/TakePresetActions.cpp
@@ -59,9 +59,12 @@ void FTakePresetActions::OpenAssetEditor(const TArray<UObject*>& InObjects, TSha
 			TSharedPtr<FTakePresetToolkit> Toolkit = MakeShared<FTakePresetToolkit>();
 			Toolkit->Initialize(EToolkitMode::WorldCentric, EditWithinLevelEditor, TakePreset);
 
-			TSharedRef<SDockTab> DockTab = EditWithinLevelEditor->GetTabManager()->InvokeTab(ITakeRecorderModule::TakeRecorderTabName);
-			TSharedRef<STakeRecorderTabContent> TabContent = StaticCastSharedRef<STakeRecorderTabContent>(DockTab->GetContent());
-			TabContent->SetupForEditing(Toolkit);
+			TSharedPtr<SDockTab> DockTab = EditWithinLevelEditor->GetTabManager()->TryInvokeTab(ITakeRecorderModule::TakeRecorderTabName);
+			if (DockTab.IsValid())
+			{
+				TSharedRef<STakeRecorderTabContent> TabContent = StaticCastSharedRef<STakeRecorderTabContent>(DockTab->GetContent());
+				TabContent->SetupForEditing(Toolkit);
+			}
 		}
 	}
 }
diff --git a/Engine/Plugins/VirtualProduction/Takes/Source/TakeRecorder/Private/TakeRecorderModule.cpp b/Engine/Plugins/VirtualProduction/Takes/Source/TakeRecorder/Private/TakeRecorderModule.cpp
--- a/Engine/Plugins/VirtualProduction/Takes/Source/TakeRecorder/Private/TakeRecorderModule.cpp
+++ b/Engine/Plugins/VirtualProduction/Takes/Source/TakeRecorder/Private/TakeRecorderModule.cpp
@@ -200,18 +200,21 @@ void FTakeRecorderModule::RegisterMenus()
 					[LevelSequence]
 					{
 						FLevelEditorModule& LevelEditorModule = FModuleManager::LoadModuleChecked<FLevelEditorModule>("LevelEditor");
-						TSharedRef<SDockTab> DockTab = LevelEditorModule.GetLevelEditorTabManager()->InvokeTab(ITakeRecorderModule::TakeRecorderTabName);
-						TSharedRef<STakeRecorderTabContent> TabContent = StaticCastSharedRef<STakeRecorderTabContent>(DockTab->GetContent());
-
-						// If this sequence has already been recorded, set it up for viewing, otherwise start recording from it.
-						UTakeMetaData* TakeMetaData = LevelSequence->FindMetaData<UTakeMetaData>();
-						if (!TakeMetaData || !TakeMetaData->Recorded())
-						{
-							TabContent->SetupForRecording(LevelSequence);
-						}
-						else
+						TSharedPtr<SDockTab> DockTab = LevelEditorModule.GetLevelEditorTabManager()->TryInvokeTab(ITakeRecorderModule::TakeRecorderTabName);
+						if (DockTab.IsValid())
 						{
-							TabContent->SetupForViewing(LevelSequence);
+							TSharedRef<STakeRecorderTabContent> TabContent = StaticCastSharedRef<STakeRecorderTabContent>(DockTab->GetContent());
+
+							// If this sequence has already been recorded, set it up for viewing, otherwise start recording from it.
+							UTakeMetaData* TakeMetaData = LevelSequence->FindMetaData<UTakeMetaData>();
+							if (!TakeMetaData || !TakeMetaData->Recorded())
+							{
+								TabContent->SetupForRecording(LevelSequence);
+							}
+							else
+							{
+								TabContent->SetupForViewing(LevelSequence);
+							}
 						}
 					}
 				)
diff --git a/Engine/Plugins/VirtualProduction/Takes/Source/TakeRecorder/Private/Widgets/STakeRecorderCockpit.cpp b/Engine/Plugins/VirtualProduction/Takes/Source/TakeRecorder/Private/Widgets/STakeRecorderCockpit.cpp
--- a/Engine/Plugins/VirtualProduction/Takes/Source/TakeRecorder/Private/Widgets/STakeRecorderCockpit.cpp
+++ b/Engine/Plugins/VirtualProduction/Takes/Source/TakeRecorder/Private/Widgets/STakeRecorderCockpit.cpp
@@ -959,9 +959,12 @@ FReply STakeRecorderCockpit::NewRecordingFromThis()
 	}
 
 	FLevelEditorModule& LevelEditorModule = FModuleManager::LoadModuleChecked<FLevelEditorModule>("LevelEditor");
-	TSharedRef<SDockTab> DockTab = LevelEditorModule.GetLevelEditorTabManager()->InvokeTab(ITakeRecorderModule::TakeRecorderTabName);
-	TSharedRef<STakeRecorderTabContent> TabContent = StaticCastSharedRef<STakeRecorderTabContent>(DockTab->GetContent());
-	TabContent->SetupForRecording(Sequence);
+	TSharedPtr<SDockTab> DockTab = LevelEditorModule.GetLevelEditorTabManager()->TryInvokeTab(ITakeRecorderModule::TakeRecorderTabName);
+	if (DockTab.IsValid())
+	{
+		TSharedRef<STakeRecorderTabContent> TabContent = StaticCastSharedRef<STakeRecorderTabContent>(DockTab->GetContent());
+		TabContent->SetupForRecording(Sequence);
+	}
 
 	return FReply::Handled();
 }
diff --git a/Engine/Plugins/VirtualProduction/Takes/Source/TakeRecorder/Private/Widgets/STakeRecorderPanel.cpp b/Engine/Plugins/VirtualProduction/Takes/Source/TakeRecorder/Private/Widgets/STakeRecorderPanel.cpp
--- a/Engine/Plugins/VirtualProduction/Takes/Source/TakeRecorder/Private/Widgets/STakeRecorderPanel.cpp
+++ b/Engine/Plugins/VirtualProduction/Takes/Source/TakeRecorder/Private/Widgets/STakeRecorderPanel.cpp
@@ -870,7 +870,7 @@ void STakeRecorderPanel::ToggleTakeBrowserCheckState(ECheckBoxState CheckState)
 	}
 	else 
 	{
-		TakesBrowserTab = LevelEditorModule.GetLevelEditorTabManager()->InvokeTab(ITakeRecorderModule::TakesBrowserTabName);
+		TakesBrowserTab = LevelEditorModule.GetLevelEditorTabManager()->TryInvokeTab(ITakeRecorderModule::TakesBrowserTabName);
 
 		bool bAllowLockedBrowser =  true;
 		bool bFocusContentBrowser = false;
diff --git a/Engine/Source/Developer/Merge/Private/Merge.cpp b/Engine/Source/Developer/Merge/Private/Merge.cpp
--- a/Engine/Source/Developer/Merge/Private/Merge.cpp
+++ b/Engine/Source/Developer/Merge/Private/Merge.cpp
@@ -98,8 +98,8 @@ class FMerge : public IMerge
 	/** IModuleInterface implementation */
 	virtual void StartupModule() override;
 	virtual void ShutdownModule() override;
-	virtual TSharedRef<SDockTab> GenerateMergeWidget(const UBlueprint& Object, TSharedRef<FBlueprintEditor> Editor) override;
-	virtual TSharedRef<SDockTab> GenerateMergeWidget(const UBlueprint* BaseAsset, const UBlueprint* RemoteAsset, const UBlueprint* LocalAsset, const FOnMergeResolved& MergeResolutionCallback, TSharedRef<FBlueprintEditor> Editor) override;
+	virtual TSharedPtr<SDockTab> GenerateMergeWidget(const UBlueprint& Object, TSharedRef<FBlueprintEditor> Editor) override;
+	virtual TSharedPtr<SDockTab> GenerateMergeWidget(const UBlueprint* BaseAsset, const UBlueprint* RemoteAsset, const UBlueprint* LocalAsset, const FOnMergeResolved& MergeResolutionCallback, TSharedRef<FBlueprintEditor> Editor) override;
 	virtual bool PendingMerge(const UBlueprint& BlueprintObj) const override;
 	//virtual FOnMergeResolved& OnMergeResolved() const override;
 
@@ -114,7 +114,7 @@ void FMerge::StartupModule()
 	// This code will execute after your module is loaded into memory (but after global variables are initialized, of course.)
 
 
-	// Registering a nomad spawner that spawns an empty dock tab on purpose - allows us to call InvokeTab() using our TabId later and set the content. (see GenerateMergeWidget())
+	// Registering a nomad spawner that spawns an empty dock tab on purpose - allows us to call TryInvokeTab() using our TabId later and set the content. (see GenerateMergeWidget())
 	FGlobalTabmanager::Get()->RegisterNomadTabSpawner(MergeToolTabId, FOnSpawnTab::CreateStatic([] (const FSpawnTabArgs&) { return SNew(SDockTab); }))
 		.SetDisplayName(NSLOCTEXT("MergeTool", "TabTitle", "Merge Tool"))
 		.SetTooltipText(NSLOCTEXT("MergeTool", "TooltipText", "Used to display several versions of a blueprint that need to be merged into a single version."))
@@ -128,15 +128,22 @@ void FMerge::ShutdownModule()
 	FGlobalTabmanager::Get()->UnregisterNomadTabSpawner(MergeToolTabId);
 }
 
-TSharedRef<SDockTab> FMerge::GenerateMergeWidget(const UBlueprint& Object, TSharedRef<FBlueprintEditor> Editor)
+TSharedPtr<SDockTab> FMerge::GenerateMergeWidget(const UBlueprint& Object, TSharedRef<FBlueprintEditor> Editor)
 {
-	auto ActiveTabPtr = ActiveTab.Pin();
+	TSharedPtr<SDockTab> ActiveTabPtr = ActiveTab.Pin();
 	if( ActiveTabPtr.IsValid() )
 	{
 		// just bring the tab to the foreground:
-		auto CurrentTab = FGlobalTabmanager::Get()->InvokeTab(MergeToolTabId);
-		check( CurrentTab == ActiveTabPtr );
-		return ActiveTabPtr.ToSharedRef();
+		TSharedPtr<SDockTab> CurrentTab = FGlobalTabmanager::Get()->TryInvokeTab(MergeToolTabId);
+		if (CurrentTab.IsValid())
+		{
+			check(CurrentTab == ActiveTabPtr);
+			return ActiveTabPtr;
+		}
+		else
+		{
+			return nullptr;
+		}
 	}
 
 	// merge the local asset with the depot, SCC provides us with the last common revision as
@@ -182,30 +189,44 @@ TSharedRef<SDockTab> FMerge::GenerateMergeWidget(const UBlueprint& Object, TShar
 		}
 	}
 
-	TSharedRef<SDockTab> Tab =  FGlobalTabmanager::Get()->InvokeTab(MergeToolTabId);
-	Tab->SetContent(Contents.ToSharedRef());
-	ActiveTab = Tab;
-	return Tab;
+	TSharedPtr<SDockTab> Tab = FGlobalTabmanager::Get()->TryInvokeTab(MergeToolTabId);
+	if (Tab.IsValid())
+	{
+		Tab->SetContent(Contents.ToSharedRef());
+		ActiveTab = Tab;
+	}
 
+	return Tab;
 }
 
-TSharedRef<SDockTab> FMerge::GenerateMergeWidget(const UBlueprint* BaseBlueprint, const UBlueprint* RemoteBlueprint, const UBlueprint* LocalBlueprint, const FOnMergeResolved& MergeResolutionCallback, TSharedRef<FBlueprintEditor> Editor)
+TSharedPtr<SDockTab> FMerge::GenerateMergeWidget(const UBlueprint* BaseBlueprint, const UBlueprint* RemoteBlueprint, const UBlueprint* LocalBlueprint, const FOnMergeResolved& MergeResolutionCallback, TSharedRef<FBlueprintEditor> Editor)
 {
 	if (ActiveTab.IsValid())
 	{
 		TSharedPtr<SDockTab> ActiveTabPtr = ActiveTab.Pin();
 		// just bring the tab to the foreground:
-		TSharedRef<SDockTab> CurrentTab = FGlobalTabmanager::Get()->InvokeTab(MergeToolTabId);
-		check(CurrentTab == ActiveTabPtr);
-		return ActiveTabPtr.ToSharedRef();
+		TSharedPtr<SDockTab> CurrentTab = FGlobalTabmanager::Get()->TryInvokeTab(MergeToolTabId);
+		if (CurrentTab.IsValid())
+		{
+			check(CurrentTab == ActiveTabPtr);
+			return ActiveTabPtr;
+		}
+		else
+		{
+			return nullptr;
+		}
 	}
 
 	// @TODO: pipe revision info through
 	TSharedPtr<SWidget> TabContents = GenerateMergeTabContents(Editor, BaseBlueprint, FRevisionInfo::InvalidRevision(), RemoteBlueprint, FRevisionInfo::InvalidRevision(), LocalBlueprint, MergeResolutionCallback);
 
-	TSharedRef<SDockTab> Tab = FGlobalTabmanager::Get()->InvokeTab(MergeToolTabId);
-	Tab->SetContent(TabContents.ToSharedRef());
-	ActiveTab = Tab;
+	TSharedPtr<SDockTab> Tab = FGlobalTabmanager::Get()->TryInvokeTab(MergeToolTabId);
+	if (Tab.IsValid())
+	{
+		Tab->SetContent(TabContents.ToSharedRef());
+		ActiveTab = Tab;
+	}
+
 	return Tab;
 }
 
diff --git a/Engine/Source/Developer/Merge/Public/Merge.h b/Engine/Source/Developer/Merge/Public/Merge.h
--- a/Engine/Source/Developer/Merge/Public/Merge.h
+++ b/Engine/Source/Developer/Merge/Public/Merge.h
@@ -53,8 +53,8 @@ public:
 	 *
 	 * @return The merge view widget
 	 */
-	virtual TSharedRef<class SDockTab> GenerateMergeWidget( const UBlueprint& Object, TSharedRef< class FBlueprintEditor> Editor ) = 0;
-	virtual TSharedRef<class SDockTab> GenerateMergeWidget(const UBlueprint* BaseBlueprint, const UBlueprint* RemoteBlueprint, const UBlueprint* LocalBlueprint, const FOnMergeResolved& MergeResolutionCallback, TSharedRef<class FBlueprintEditor> Editor) = 0;
+	virtual TSharedPtr<class SDockTab> GenerateMergeWidget( const UBlueprint& Object, TSharedRef< class FBlueprintEditor> Editor ) = 0;
+	virtual TSharedPtr<class SDockTab> GenerateMergeWidget(const UBlueprint* BaseBlueprint, const UBlueprint* RemoteBlueprint, const UBlueprint* LocalBlueprint, const FOnMergeResolved& MergeResolutionCallback, TSharedRef<class FBlueprintEditor> Editor) = 0;
 
 	/** 
 	 * @return whether the blueprint is in a conflicted state, and can therefore be merged.
diff --git a/Engine/Source/Developer/MessageLog/Private/MessageLogModule.cpp b/Engine/Source/Developer/MessageLog/Private/MessageLogModule.cpp
--- a/Engine/Source/Developer/MessageLog/Private/MessageLogModule.cpp
+++ b/Engine/Source/Developer/MessageLog/Private/MessageLogModule.cpp
@@ -124,7 +124,7 @@ void FMessageLogModule::OpenMessageLog(const FName& LogName)
 #endif
 	   )
 	{
-		FGlobalTabmanager::Get()->InvokeTab(FName("MessageLog"));
+		FGlobalTabmanager::Get()->TryInvokeTab(FName("MessageLog"));
 		MessageLogViewModel->ChangeCurrentListingViewModel(LogName);
 	}
 }
diff --git a/Engine/Source/Developer/ProjectLauncher/Private/Widgets/Deploy/SProjectLauncherSimpleDeviceListView.cpp b/Engine/Source/Developer/ProjectLauncher/Private/Widgets/Deploy/SProjectLauncherSimpleDeviceListView.cpp
--- a/Engine/Source/Developer/ProjectLauncher/Private/Widgets/Deploy/SProjectLauncherSimpleDeviceListView.cpp
+++ b/Engine/Source/Developer/ProjectLauncher/Private/Widgets/Deploy/SProjectLauncherSimpleDeviceListView.cpp
@@ -127,7 +127,7 @@ bool SProjectLauncherSimpleDeviceListView::HandleDeviceListRowIsEnabled(TSharedP
 
 void SProjectLauncherSimpleDeviceListView::HandleDeviceManagerHyperlinkNavigate() const
 {
-	FGlobalTabmanager::Get()->InvokeTab(FTabId("DeviceManager"));
+	FGlobalTabmanager::Get()->TryInvokeTab(FTabId("DeviceManager"));
 }
 
 
diff --git a/Engine/Source/Developer/SessionFrontend/Private/SessionFrontendModule.cpp b/Engine/Source/Developer/SessionFrontend/Private/SessionFrontendModule.cpp
--- a/Engine/Source/Developer/SessionFrontend/Private/SessionFrontendModule.cpp
+++ b/Engine/Source/Developer/SessionFrontend/Private/SessionFrontendModule.cpp
@@ -41,12 +41,12 @@ public:
 
 	virtual void InvokeSessionFrontend(FName SubTabToActivate = NAME_None) override
 	{
-		FGlobalTabmanager::Get()->InvokeTab(SessionFrontendTabName);
+		FGlobalTabmanager::Get()->TryInvokeTab(SessionFrontendTabName);
 		if ( WeakFrontend.IsValid() )
 		{
 			if ( SubTabToActivate != NAME_None )
 			{
-				WeakFrontend.Pin()->GetTabManager()->InvokeTab(SubTabToActivate);
+				WeakFrontend.Pin()->GetTabManager()->TryInvokeTab(SubTabToActivate);
 			}
 		}
 	}
diff --git a/Engine/Source/Developer/SlateReflector/Private/SlateReflectorModule.cpp b/Engine/Source/Developer/SlateReflector/Private/SlateReflectorModule.cpp
--- a/Engine/Source/Developer/SlateReflector/Private/SlateReflectorModule.cpp
+++ b/Engine/Source/Developer/SlateReflector/Private/SlateReflectorModule.cpp
@@ -95,19 +95,19 @@ public:
 	virtual void DisplayWidgetReflector() override
 	{
 		check(bHasRegisteredTabSpawners);
-		FGlobalTabmanager::Get()->InvokeTab(FTabId("WidgetReflector"));
+		FGlobalTabmanager::Get()->TryInvokeTab(FTabId("WidgetReflector"));
 	}
 
 	virtual void DisplayTextureAtlasVisualizer() override
 	{
 		check(bHasRegisteredTabSpawners);
-		FGlobalTabmanager::Get()->InvokeTab(FTabId("TextureAtlasVisualizer"));
+		FGlobalTabmanager::Get()->TryInvokeTab(FTabId("TextureAtlasVisualizer"));
 	}
 
 	virtual void DisplayFontAtlasVisualizer() override
 	{
 		check(bHasRegisteredTabSpawners);
-		FGlobalTabmanager::Get()->InvokeTab(FTabId("FontAtlasVisualizer"));
+		FGlobalTabmanager::Get()->TryInvokeTab(FTabId("FontAtlasVisualizer"));
 	}
 
 	virtual void RegisterTabSpawner( const TSharedPtr<FWorkspaceItem>& WorkspaceGroup ) override
diff --git a/Engine/Source/Developer/SlateReflector/Private/Widgets/SWidgetReflector.cpp b/Engine/Source/Developer/SlateReflector/Private/Widgets/SWidgetReflector.cpp
--- a/Engine/Source/Developer/SlateReflector/Private/Widgets/SWidgetReflector.cpp
+++ b/Engine/Source/Developer/SlateReflector/Private/Widgets/SWidgetReflector.cpp
@@ -1049,7 +1049,7 @@ void SWidgetReflector::SetUIMode(const EWidgetReflectorUIMode InNewMode)
 
 		if (CurrentUIMode == EWidgetReflectorUIMode::Snapshot)
 		{
-			TabManager->InvokeTab(WidgetReflectorTabID::SnapshotWidgetPicker);
+			TabManager->TryInvokeTab(WidgetReflectorTabID::SnapshotWidgetPicker);
 		}
 		else
 		{
@@ -1738,7 +1738,7 @@ void SWidgetReflector::HandleReflectorTreeSelectionChanged( TSharedPtr<FWidgetRe
 
 	if (SelectedWidgetObjects.Num() > 0)
 	{
-		TabManager->InvokeTab(WidgetReflectorTabID::WidgetDetails);
+		TabManager->TryInvokeTab(WidgetReflectorTabID::WidgetDetails);
 		if (PropertyViewPtr.IsValid())
 		{
 			PropertyViewPtr->SetObjects(SelectedWidgetObjects);
diff --git a/Engine/Source/Developer/TraceInsights/Private/Insights/InsightsManager.cpp b/Engine/Source/Developer/TraceInsights/Private/Insights/InsightsManager.cpp
--- a/Engine/Source/Developer/TraceInsights/Private/Insights/InsightsManager.cpp
+++ b/Engine/Source/Developer/TraceInsights/Private/Insights/InsightsManager.cpp
@@ -156,8 +156,8 @@ bool FInsightsManager::Tick(float DeltaTime)
 
 			if (FGlobalTabmanager::Get()->HasTabSpawner(FInsightsManagerTabs::NetworkingProfilerTabId))
 			{
-				FGlobalTabmanager::Get()->InvokeTab(FInsightsManagerTabs::NetworkingProfilerTabId);
-				FGlobalTabmanager::Get()->InvokeTab(FInsightsManagerTabs::NetworkingProfilerTabId);
+				FGlobalTabmanager::Get()->TryInvokeTab(FInsightsManagerTabs::NetworkingProfilerTabId);
+				FGlobalTabmanager::Get()->TryInvokeTab(FInsightsManagerTabs::NetworkingProfilerTabId);
 			}
 
 			ActivateTimingInsightsTab();
@@ -225,13 +225,13 @@ void FInsightsManager::SpawnAndActivateTabs()
 	// Open Timing Insights tab.
 	if (FGlobalTabmanager::Get()->HasTabSpawner(FInsightsManagerTabs::TimingProfilerTabId))
 	{
-		FGlobalTabmanager::Get()->InvokeTab(FInsightsManagerTabs::TimingProfilerTabId);
+		FGlobalTabmanager::Get()->TryInvokeTab(FInsightsManagerTabs::TimingProfilerTabId);
 	}
 
 	// Open Asset Loading Insights tab.
 	if (FGlobalTabmanager::Get()->HasTabSpawner(FInsightsManagerTabs::LoadingProfilerTabId))
 	{
-		FGlobalTabmanager::Get()->InvokeTab(FInsightsManagerTabs::LoadingProfilerTabId);
+		FGlobalTabmanager::Get()->TryInvokeTab(FInsightsManagerTabs::LoadingProfilerTabId);
 	}
 
 	// Close the existing Networking Insights tabs.
diff --git a/Engine/Source/Developer/TraceInsights/Private/Insights/LoadingProfiler/Widgets/SLoadingProfilerWindow.cpp b/Engine/Source/Developer/TraceInsights/Private/Insights/LoadingProfiler/Widgets/SLoadingProfilerWindow.cpp
--- a/Engine/Source/Developer/TraceInsights/Private/Insights/LoadingProfiler/Widgets/SLoadingProfilerWindow.cpp
+++ b/Engine/Source/Developer/TraceInsights/Private/Insights/LoadingProfiler/Widgets/SLoadingProfilerWindow.cpp
@@ -545,7 +545,7 @@ void SLoadingProfilerWindow::ShowTab(const FName& TabID)
 {
 	if (TabManager->HasTabSpawner(TabID))
 	{
-		TabManager->InvokeTab(TabID);
+		TabManager->TryInvokeTab(TabID);
 	}
 }
 
diff --git a/Engine/Source/Developer/TraceInsights/Private/Insights/NetworkingProfiler/Widgets/SNetworkingProfilerWindow.cpp b/Engine/Source/Developer/TraceInsights/Private/Insights/NetworkingProfiler/Widgets/SNetworkingProfilerWindow.cpp
--- a/Engine/Source/Developer/TraceInsights/Private/Insights/NetworkingProfiler/Widgets/SNetworkingProfilerWindow.cpp
+++ b/Engine/Source/Developer/TraceInsights/Private/Insights/NetworkingProfiler/Widgets/SNetworkingProfilerWindow.cpp
@@ -418,7 +418,7 @@ void SNetworkingProfilerWindow::ShowTab(const FName& TabID)
 {
 	if (TabManager->HasTabSpawner(TabID))
 	{
-		TabManager->InvokeTab(TabID);
+		TabManager->TryInvokeTab(TabID);
 	}
 }
 
diff --git a/Engine/Source/Developer/TraceInsights/Private/Insights/Widgets/STimingProfilerWindow.cpp b/Engine/Source/Developer/TraceInsights/Private/Insights/Widgets/STimingProfilerWindow.cpp
--- a/Engine/Source/Developer/TraceInsights/Private/Insights/Widgets/STimingProfilerWindow.cpp
+++ b/Engine/Source/Developer/TraceInsights/Private/Insights/Widgets/STimingProfilerWindow.cpp
@@ -586,7 +586,7 @@ void STimingProfilerWindow::ShowTab(const FName& TabID)
 {
 	if (TabManager->HasTabSpawner(TabID))
 	{
-		TabManager->InvokeTab(TabID);
+		TabManager->TryInvokeTab(TabID);
 	}
 }
 
diff --git a/Engine/Source/Editor/AnimationBlueprintEditor/Private/AnimationBlueprintEditorMode.cpp b/Engine/Source/Editor/AnimationBlueprintEditor/Private/AnimationBlueprintEditorMode.cpp
--- a/Engine/Source/Editor/AnimationBlueprintEditor/Private/AnimationBlueprintEditorMode.cpp
+++ b/Engine/Source/Editor/AnimationBlueprintEditor/Private/AnimationBlueprintEditorMode.cpp
@@ -187,7 +187,7 @@ void FAnimationBlueprintEditorMode::PostActivateMode()
 		// If we are a derived anim blueprint always show the overrides tab
 		if(UAnimBlueprint::FindRootAnimBlueprint(AnimBlueprint))
 		{
-			MyBlueprintEditor.Pin()->GetTabManager()->InvokeTab(AnimationBlueprintEditorTabs::AssetOverridesTab);
+			MyBlueprintEditor.Pin()->GetTabManager()->TryInvokeTab(AnimationBlueprintEditorTabs::AssetOverridesTab);
 		}
 	}
 
diff --git a/Engine/Source/Editor/BehaviorTreeEditor/Private/BehaviorTreeEditor.cpp b/Engine/Source/Editor/BehaviorTreeEditor/Private/BehaviorTreeEditor.cpp
--- a/Engine/Source/Editor/BehaviorTreeEditor/Private/BehaviorTreeEditor.cpp
+++ b/Engine/Source/Editor/BehaviorTreeEditor/Private/BehaviorTreeEditor.cpp
@@ -946,7 +946,7 @@ void FBehaviorTreeEditor::BindCommonCommands()
 
 void FBehaviorTreeEditor::SearchTree()
 {
-	TabManager->InvokeTab(FBehaviorTreeEditorTabs::SearchID);
+	TabManager->TryInvokeTab(FBehaviorTreeEditorTabs::SearchID);
 	FindResults->FocusForUse();
 }
 
diff --git a/Engine/Source/Editor/Blutility/Private/EditorUtilitySubsystem.cpp b/Engine/Source/Editor/Blutility/Private/EditorUtilitySubsystem.cpp
--- a/Engine/Source/Editor/Blutility/Private/EditorUtilitySubsystem.cpp
+++ b/Engine/Source/Editor/Blutility/Private/EditorUtilitySubsystem.cpp
@@ -148,7 +148,7 @@ bool UEditorUtilitySubsystem::SpawnRegisteredTabByID(FName NewTabID)
 		TSharedPtr<FTabManager> LevelEditorTabManager = LevelEditorModule.GetLevelEditorTabManager();
 		if (LevelEditorTabManager->HasTabSpawner(NewTabID))
 		{
-			TSharedRef<SDockTab> NewDockTab = LevelEditorTabManager->InvokeTab(NewTabID);
+			TSharedPtr<SDockTab> NewDockTab = LevelEditorTabManager->TryInvokeTab(NewTabID);
 			return true;
 		}
 	}
diff --git a/Engine/Source/Editor/ClothPainter/Private/ClothPaintingModule.cpp b/Engine/Source/Editor/ClothPainter/Private/ClothPaintingModule.cpp
--- a/Engine/Source/Editor/ClothPainter/Private/ClothPaintingModule.cpp
+++ b/Engine/Source/Editor/ClothPainter/Private/ClothPaintingModule.cpp
@@ -185,7 +185,7 @@ TSharedPtr<SClothPaintTab> FClothPaintingModule::GetActiveClothTab(TWeakPtr<ISke
 
 	if(bInvoke)
 	{
-		TabManager->InvokeTab(FTabId(FClothPaintTabSummoner::TabName));
+		TabManager->TryInvokeTab(FTabId(FClothPaintTabSummoner::TabName));
 	}
 
 	// If we can't summon this tab we will have spawned a placeholder which we can not cast as done below,
diff --git a/Engine/Source/Editor/ContentBrowser/Private/ContentBrowserSingleton.cpp b/Engine/Source/Editor/ContentBrowser/Private/ContentBrowserSingleton.cpp
--- a/Engine/Source/Editor/ContentBrowser/Private/ContentBrowserSingleton.cpp
+++ b/Engine/Source/Editor/ContentBrowser/Private/ContentBrowserSingleton.cpp
@@ -549,7 +549,7 @@ void FContentBrowserSingleton::FocusContentBrowser(const TSharedPtr<SContentBrow
 		TSharedPtr<FTabManager> TabManager = Browser->GetTabManager();
 		if ( TabManager.IsValid() )
 		{
-			TabManager->InvokeTab(Browser->GetInstanceName());
+			TabManager->TryInvokeTab(Browser->GetInstanceName());
 		}
 	}
 }
@@ -585,11 +585,11 @@ FName FContentBrowserSingleton::SummonNewBrowser(bool bAllowLockedBrowsers)
 		const TWeakPtr<FTabManager>& TabManagerToInvoke = BrowserToLastKnownTabManagerMap.FindRef(NewTabName);
 		if ( TabManagerToInvoke.IsValid() )
 		{
-			TabManagerToInvoke.Pin()->InvokeTab(NewTabName);
+			TabManagerToInvoke.Pin()->TryInvokeTab(NewTabName);
 		}
 		else
 		{
-			FGlobalTabmanager::Get()->InvokeTab(NewTabName);
+			FGlobalTabmanager::Get()->TryInvokeTab(NewTabName);
 		}
 	}
 	else
diff --git a/Engine/Source/Editor/ContentBrowser/Private/SAssetView.cpp b/Engine/Source/Editor/ContentBrowser/Private/SAssetView.cpp
--- a/Engine/Source/Editor/ContentBrowser/Private/SAssetView.cpp
+++ b/Engine/Source/Editor/ContentBrowser/Private/SAssetView.cpp
@@ -4530,7 +4530,7 @@ void SAssetView::AssetRenameCommit(const TSharedPtr<FAssetViewItem>& Item, const
 			//Notify the user rename fail and link the output log
 			FNotificationInfo Info(LOCTEXT("RenameAssetsFailed", "Failed to rename assets"));
 			Info.ExpireDuration = 5.0f;
-			Info.Hyperlink = FSimpleDelegate::CreateStatic([](){ FGlobalTabmanager::Get()->InvokeTab(FName("OutputLog")); });
+			Info.Hyperlink = FSimpleDelegate::CreateStatic([](){ FGlobalTabmanager::Get()->TryInvokeTab(FName("OutputLog")); });
 			Info.HyperlinkText = LOCTEXT("ShowOutputLogHyperlink", "Show Output Log");
 			FSlateNotificationManager::Get().AddNotification(Info);
 			
diff --git a/Engine/Source/Editor/DeviceProfileEditor/Private/SDeviceProfileEditor.cpp b/Engine/Source/Editor/DeviceProfileEditor/Private/SDeviceProfileEditor.cpp
--- a/Engine/Source/Editor/DeviceProfileEditor/Private/SDeviceProfileEditor.cpp
+++ b/Engine/Source/Editor/DeviceProfileEditor/Private/SDeviceProfileEditor.cpp
@@ -483,7 +483,7 @@ void SDeviceProfileEditor::HandleDeviceProfilePinned( const TWeakObjectPtr< UDev
 		DeviceProfiles.Add( DeviceProfile.Get() );
 		RebuildPropertyTable();
 
-		TabManager->InvokeTab(DeviceProfileEditorTabName);
+		TabManager->TryInvokeTab(DeviceProfileEditorTabName);
 	}
 }
 
@@ -495,7 +495,7 @@ void SDeviceProfileEditor::HandleDeviceProfileUnpinned( const TWeakObjectPtr< UD
 		DeviceProfiles.Remove( DeviceProfile.Get() );
 		RebuildPropertyTable();
 
-		TabManager->InvokeTab(DeviceProfileEditorTabName);
+		TabManager->TryInvokeTab(DeviceProfileEditorTabName);
 	}
 }
 
@@ -514,7 +514,7 @@ void SDeviceProfileEditor::HandleDeviceProfileViewAlone( const TWeakObjectPtr< U
 			.SetGroup(DeviceManagerMenuGroup.ToSharedRef());
 	}
 
-	TabManager->InvokeTab(TabId);
+	TabManager->TryInvokeTab(TabId);
 }
 
 
diff --git a/Engine/Source/Editor/EditorSettingsViewer/Private/EditorSettingsViewerModule.cpp b/Engine/Source/Editor/EditorSettingsViewer/Private/EditorSettingsViewerModule.cpp
--- a/Engine/Source/Editor/EditorSettingsViewer/Private/EditorSettingsViewerModule.cpp
+++ b/Engine/Source/Editor/EditorSettingsViewer/Private/EditorSettingsViewerModule.cpp
@@ -64,7 +64,7 @@ public:
 
 	virtual void ShowSettings( const FName& CategoryName, const FName& SectionName ) override
 	{
-		FGlobalTabmanager::Get()->InvokeTab(EditorSettingsTabName);
+		FGlobalTabmanager::Get()->TryInvokeTab(EditorSettingsTabName);
 		ISettingsEditorModelPtr SettingsEditorModel = SettingsEditorModelPtr.Pin();
 
 		if (SettingsEditorModel.IsValid())
diff --git a/Engine/Source/Editor/FontEditor/Private/FontEditor.cpp b/Engine/Source/Editor/FontEditor/Private/FontEditor.cpp
--- a/Engine/Source/Editor/FontEditor/Private/FontEditor.cpp
+++ b/Engine/Source/Editor/FontEditor/Private/FontEditor.cpp
@@ -549,13 +549,13 @@ void FFontEditor::UpdateLayout()
 	switch(Font->FontCacheType)
 	{
 	case EFontCacheType::Offline:
-		TabManager->InvokeTab(TexturePagesViewportTabId);
-		TabManager->InvokeTab(PagePropertiesTabId);
+		TabManager->TryInvokeTab(TexturePagesViewportTabId);
+		TabManager->TryInvokeTab(PagePropertiesTabId);
 		CloseTab(CompositeFontEditorTabId);
 		break;
 
 	case EFontCacheType::Runtime:
-		TabManager->InvokeTab(CompositeFontEditorTabId);
+		TabManager->TryInvokeTab(CompositeFontEditorTabId);
 		CloseTab(TexturePagesViewportTabId);
 		CloseTab(PagePropertiesTabId);
 		break;
diff --git a/Engine/Source/Editor/GameProjectGeneration/Private/SNewClassDialog.cpp b/Engine/Source/Editor/GameProjectGeneration/Private/SNewClassDialog.cpp
--- a/Engine/Source/Editor/GameProjectGeneration/Private/SNewClassDialog.cpp
+++ b/Engine/Source/Editor/GameProjectGeneration/Private/SNewClassDialog.cpp
@@ -1336,7 +1336,7 @@ void SNewClassDialog::FinishClicked()
 				, FText::FromString(NewClassName), FText::FromString(SelectedModuleInfo->ModuleName), FailReason );
 			if( FMessageDialog::Open(EAppMsgType::YesNo, Message) == EAppReturnType::Yes )
 			{
-				FGlobalTabmanager::Get()->InvokeTab(FName("OutputLog"));
+				FGlobalTabmanager::Get()->TryInvokeTab(FName("OutputLog"));
 			}
 
 			// We did manage to add the code itself, so we can close the dialog.
diff --git a/Engine/Source/Editor/IntroTutorials/Private/EditorTutorial.cpp b/Engine/Source/Editor/IntroTutorials/Private/EditorTutorial.cpp
--- a/Engine/Source/Editor/IntroTutorials/Private/EditorTutorial.cpp
+++ b/Engine/Source/Editor/IntroTutorials/Private/EditorTutorial.cpp
@@ -87,7 +87,7 @@ void UEditorTutorial::HandleTutorialLaunched()
 {
 	FEditorScriptExecutionGuard ScriptGuard;
 	FLevelEditorModule& LevelEditorModule = FModuleManager::GetModuleChecked<FLevelEditorModule>(TEXT("LevelEditor"));
-	LevelEditorModule.GetLevelEditorTabManager()->InvokeTab(FTabId("TutorialsBrowser"))->RequestCloseTab();
+	LevelEditorModule.GetLevelEditorTabManager()->TryInvokeTab(FTabId("TutorialsBrowser"))->RequestCloseTab();
 	OnTutorialLaunched();
 }
 
diff --git a/Engine/Source/Editor/IntroTutorials/Private/IntroTutorials.cpp b/Engine/Source/Editor/IntroTutorials/Private/IntroTutorials.cpp
--- a/Engine/Source/Editor/IntroTutorials/Private/IntroTutorials.cpp
+++ b/Engine/Source/Editor/IntroTutorials/Private/IntroTutorials.cpp
@@ -416,7 +416,7 @@ void FIntroTutorials::SummonTutorialBrowser()
 	if(TutorialRoot.IsValid())
 	{
 		FLevelEditorModule& LevelEditorModule = FModuleManager::GetModuleChecked<FLevelEditorModule>( TEXT("LevelEditor") );
-		TutorialBrowserDockTab = LevelEditorModule.GetLevelEditorTabManager()->InvokeTab(FTabId("TutorialsBrowser"));
+		TutorialBrowserDockTab = LevelEditorModule.GetLevelEditorTabManager()->TryInvokeTab(FTabId("TutorialsBrowser"));
 	}
 }
 
diff --git a/Engine/Source/Editor/IntroTutorials/Private/STutorialOverlay.cpp b/Engine/Source/Editor/IntroTutorials/Private/STutorialOverlay.cpp
--- a/Engine/Source/Editor/IntroTutorials/Private/STutorialOverlay.cpp
+++ b/Engine/Source/Editor/IntroTutorials/Private/STutorialOverlay.cpp
@@ -244,7 +244,7 @@ void STutorialOverlay::OpenBrowserForWidgetAnchor(UEditorTutorial* InTutorial, c
 			// fallback to trying the main level editor tab manager
 			FLevelEditorModule& LevelEditorModule = FModuleManager::GetModuleChecked<FLevelEditorModule>(TEXT("LevelEditor"));
 			TSharedPtr<FTabManager> LevelEditorTabManager = LevelEditorModule.GetLevelEditorTabManager();
-			LevelEditorTabManager->InvokeTab(FName(*WidgetContent.WidgetAnchor.TabToFocusOrOpen));
+			LevelEditorTabManager->TryInvokeTab(FName(*WidgetContent.WidgetAnchor.TabToFocusOrOpen));
 		}
 	}
 }
diff --git a/Engine/Source/Editor/Kismet/Private/BlueprintDebugger.cpp b/Engine/Source/Editor/Kismet/Private/BlueprintDebugger.cpp
--- a/Engine/Source/Editor/Kismet/Private/BlueprintDebugger.cpp
+++ b/Engine/Source/Editor/Kismet/Private/BlueprintDebugger.cpp
@@ -205,7 +205,7 @@ TSharedRef<SDockTab> FBlueprintDebuggerImpl::CreateBluprintDebuggerTab(const FSp
 			}
 			else
 			{
-				InDebuggingToolsManager->InvokeTab(InTabName);
+				InDebuggingToolsManager->TryInvokeTab(InTabName);
 			}
 		}
 	};
diff --git a/Engine/Source/Editor/Kismet/Private/BlueprintEditor.cpp b/Engine/Source/Editor/Kismet/Private/BlueprintEditor.cpp
--- a/Engine/Source/Editor/Kismet/Private/BlueprintEditor.cpp
+++ b/Engine/Source/Editor/Kismet/Private/BlueprintEditor.cpp
@@ -869,7 +869,7 @@ void FBlueprintEditor::SummonSearchUI(bool bSetFindWithinBlueprint, FString NewS
 		|| !GetDefault<UBlueprintEditorSettings>()->bHostFindInBlueprintsInGlobalTab)
 	{
 		FindResultsToUse = FindResults;
-		TabManager->InvokeTab(FBlueprintEditorTabs::FindResultsID);
+		TabManager->TryInvokeTab(FBlueprintEditorTabs::FindResultsID);
 	}
 	else
 	{
@@ -884,7 +884,7 @@ void FBlueprintEditor::SummonSearchUI(bool bSetFindWithinBlueprint, FString NewS
 
 void FBlueprintEditor::SummonFindAndReplaceUI()
 {
-	TabManager->InvokeTab(FBlueprintEditorTabs::ReplaceNodeReferencesID);
+	TabManager->TryInvokeTab(FBlueprintEditorTabs::ReplaceNodeReferencesID);
 }
 
 void FBlueprintEditor::EnableSCSPreview(bool bEnable)
@@ -1008,7 +1008,7 @@ void FBlueprintEditor::OnComponentDoubleClicked(TSharedPtr<class FSCSEditorTreeN
 	TSharedPtr<SDockTab> OwnerTab = Inspector->GetOwnerTab();
 	if ( OwnerTab.IsValid() )
 	{
-		GetTabManager()->InvokeTab(FBlueprintEditorTabs::SCSViewportID);
+		GetTabManager()->TryInvokeTab(FBlueprintEditorTabs::SCSViewportID);
 	}
 }
 
@@ -1979,7 +1979,7 @@ void FBlueprintEditor::RegisterApplicationModes(const TArray<UBlueprint*>& InBlu
 
 				if ( bShouldOpenInComponentsMode && CanAccessComponentsMode() )
 				{
-					TabManager->InvokeTab(FBlueprintEditorTabs::SCSViewportID);
+					TabManager->TryInvokeTab(FBlueprintEditorTabs::SCSViewportID);
 				}
 			}
 		}
@@ -2249,7 +2249,7 @@ void FBlueprintEditor::PostLayoutBlueprintEditorInitialization()
 		// If we have a warning/error, open output log.
 		if (!Blueprint->IsUpToDate() || (Blueprint->Status == BS_UpToDateWithWarnings))
 		{
-			TabManager->InvokeTab(FBlueprintEditorTabs::CompilerResultsID);
+			TabManager->TryInvokeTab(FBlueprintEditorTabs::CompilerResultsID);
 		}
 	}
 
@@ -3929,7 +3929,7 @@ void FBlueprintEditor::DumpMessagesToCompilerLog(const TArray<TSharedRef<FTokeni
 	
 	if (!bEditorMarkedAsClosed && bForceMessageDisplay && GetCurrentMode() == FBlueprintEditorApplicationModes::StandardBlueprintEditorMode)
 	{
-		TabManager->InvokeTab(FBlueprintEditorTabs::CompilerResultsID);
+		TabManager->TryInvokeTab(FBlueprintEditorTabs::CompilerResultsID);
 	}
 }
 
@@ -8131,7 +8131,7 @@ void FBlueprintEditor::RefreshStandAloneDefaultsEditor()
 
 void FBlueprintEditor::RenameNewlyAddedAction(FName InActionName)
 {
-	TabManager->InvokeTab(FBlueprintEditorTabs::MyBlueprintID);
+	TabManager->TryInvokeTab(FBlueprintEditorTabs::MyBlueprintID);
 	TryInvokingDetailsTab(/*Flash*/false);
 
 	if (MyBlueprintWidget.IsValid())
@@ -9279,7 +9279,7 @@ void FBlueprintEditor::TryInvokingDetailsTab(bool bFlash)
 				if ( !Inspector.IsValid() || !Inspector->GetOwnerTab().IsValid() || Inspector->GetOwnerTab()->GetDockArea().IsValid() )
 				{
 					// Show the details panel if it doesn't exist.
-					TabManager->InvokeTab(FBlueprintEditorTabs::DetailsID);
+					TabManager->TryInvokeTab(FBlueprintEditorTabs::DetailsID);
 
 					if ( bFlash )
 					{
diff --git a/Engine/Source/Editor/Kismet/Private/FindInBlueprintManager.cpp b/Engine/Source/Editor/Kismet/Private/FindInBlueprintManager.cpp
--- a/Engine/Source/Editor/Kismet/Private/FindInBlueprintManager.cpp
+++ b/Engine/Source/Editor/Kismet/Private/FindInBlueprintManager.cpp
@@ -3404,10 +3404,16 @@ TSharedPtr<SFindInBlueprints> FFindInBlueprintSearchManager::OpenGlobalFindResul
 		const FName GlobalTabId = GlobalFindResultsTabIDs[Idx];
 		if (!OpenGlobalTabIDs.Contains(GlobalTabId))
 		{
-			TSharedRef<SDockTab> NewTab = FGlobalTabmanager::Get()->InvokeTab(GlobalTabId);
-			TSharedRef<SFindInBlueprints> NewFindTab = StaticCastSharedRef<SFindInBlueprints>(NewTab->GetContent());
-
-			return NewFindTab;
+			TSharedPtr<SDockTab> NewTab = FGlobalTabmanager::Get()->TryInvokeTab(GlobalTabId);
+			if (NewTab.IsValid())
+			{
+				TSharedRef<SFindInBlueprints> NewFindTab = StaticCastSharedRef<SFindInBlueprints>(NewTab->GetContent());
+				return NewFindTab;
+			}
+			else
+			{
+				break;
+			}
 		}
 	}
 
@@ -3430,7 +3436,7 @@ TSharedPtr<SFindInBlueprints> FFindInBlueprintSearchManager::GetGlobalFindResult
 
 	if (FindResultsToUse.IsValid())
 	{
-		FGlobalTabmanager::Get()->InvokeTab(FindResultsToUse->GetHostTabId());
+		FGlobalTabmanager::Get()->TryInvokeTab(FindResultsToUse->GetHostTabId());
 	}
 	else
 	{
diff --git a/Engine/Source/Editor/Kismet/Private/FindInBlueprints.cpp b/Engine/Source/Editor/Kismet/Private/FindInBlueprints.cpp
--- a/Engine/Source/Editor/Kismet/Private/FindInBlueprints.cpp
+++ b/Engine/Source/Editor/Kismet/Private/FindInBlueprints.cpp
@@ -417,7 +417,7 @@ FReply FFindInBlueprintsProperty::OnClick()
 			{
 				// Open Viewport Tab
 				BlueprintEditor->FocusWindow();
-				BlueprintEditor->GetTabManager()->InvokeTab(FBlueprintEditorTabs::SCSViewportID);
+				BlueprintEditor->GetTabManager()->TryInvokeTab(FBlueprintEditorTabs::SCSViewportID);
 
 				// Find and Select the Component in the Viewport tab view
 				const TArray<USCS_Node*>& Nodes = Blueprint->SimpleConstructionScript->GetAllNodes();
diff --git a/Engine/Source/Editor/Kismet/Private/SSCSEditorViewport.cpp b/Engine/Source/Editor/Kismet/Private/SSCSEditorViewport.cpp
--- a/Engine/Source/Editor/Kismet/Private/SSCSEditorViewport.cpp
+++ b/Engine/Source/Editor/Kismet/Private/SSCSEditorViewport.cpp
@@ -383,7 +383,7 @@ void SSCSEditorViewport::ToggleIsSimulateEnabled()
 	{
 		if ( GetDefault<UBlueprintEditorSettings>()->bShowViewportOnSimulate )
 		{
-			BlueprintEditorPtr.Pin()->GetTabManager()->InvokeTab(FBlueprintEditorTabs::SCSViewportID);
+			BlueprintEditorPtr.Pin()->GetTabManager()->TryInvokeTab(FBlueprintEditorTabs::SCSViewportID);
 		}
 	}
 
diff --git a/Engine/Source/Editor/LevelEditor/Private/LevelEditor.cpp b/Engine/Source/Editor/LevelEditor/Private/LevelEditor.cpp
--- a/Engine/Source/Editor/LevelEditor/Private/LevelEditor.cpp
+++ b/Engine/Source/Editor/LevelEditor/Private/LevelEditor.cpp
@@ -375,26 +375,26 @@ void FLevelEditorModule::SummonSelectionDetails()
 void FLevelEditorModule::SummonBuildAndSubmit()
 {
 	TSharedPtr<SLevelEditor> LevelEditorInstance = LevelEditorInstancePtr.Pin();
-	LevelEditorInstance->InvokeTab(LevelEditorTabIds::LevelEditorBuildAndSubmit);
+	LevelEditorInstance->TryInvokeTab(LevelEditorTabIds::LevelEditorBuildAndSubmit);
 }
 
 
 void FLevelEditorModule::SummonWorldBrowserHierarchy()
 {
 	TSharedPtr<SLevelEditor> LevelEditorInstance = LevelEditorInstancePtr.Pin();
-	LevelEditorInstance->InvokeTab(LevelEditorTabIds::WorldBrowserHierarchy);
+	LevelEditorInstance->TryInvokeTab(LevelEditorTabIds::WorldBrowserHierarchy);
 }
 
 void FLevelEditorModule::SummonWorldBrowserDetails()
 {
 	TSharedPtr<SLevelEditor> LevelEditorInstance = LevelEditorInstancePtr.Pin();
-	LevelEditorInstance->InvokeTab(LevelEditorTabIds::WorldBrowserDetails);
+	LevelEditorInstance->TryInvokeTab(LevelEditorTabIds::WorldBrowserDetails);
 }
 
 void FLevelEditorModule::SummonWorldBrowserComposition()
 {
 	TSharedPtr<SLevelEditor> LevelEditorInstance = LevelEditorInstancePtr.Pin();
-	LevelEditorInstance->InvokeTab(LevelEditorTabIds::WorldBrowserComposition);
+	LevelEditorInstance->TryInvokeTab(LevelEditorTabIds::WorldBrowserComposition);
 }
 
 // @todo remove when world-centric mode is added
diff --git a/Engine/Source/Editor/LevelEditor/Private/LevelEditorActions.cpp b/Engine/Source/Editor/LevelEditor/Private/LevelEditorActions.cpp
--- a/Engine/Source/Editor/LevelEditor/Private/LevelEditorActions.cpp
+++ b/Engine/Source/Editor/LevelEditor/Private/LevelEditorActions.cpp
@@ -2067,7 +2067,7 @@ bool FLevelEditorActionCallbacks::SaveAnimationFromSkeletalMeshComponent(AActor
 
 void FLevelEditorActionCallbacks::OpenMergeActor_Clicked()
 {
-	FGlobalTabmanager::Get()->InvokeTab(FName("MergeActors"));
+	FGlobalTabmanager::Get()->TryInvokeTab(FName("MergeActors"));
 }
 
 void FLevelEditorActionCallbacks::OnKeepSimulationChanges()
@@ -2221,7 +2221,7 @@ void FLevelEditorActionCallbacks::OnFindActorInLevelScript()
 void FLevelEditorActionCallbacks::OnShowWorldProperties( TWeakPtr< SLevelEditor > LevelEditor )
 {
 	FLevelEditorModule& LevelEditorModule = FModuleManager::GetModuleChecked<FLevelEditorModule>( TEXT("LevelEditor") );
-	LevelEditorModule.GetLevelEditorTabManager()->InvokeTab(FName("WorldSettingsTab"));
+	LevelEditorModule.GetLevelEditorTabManager()->TryInvokeTab(FName("WorldSettingsTab"));
 }
 
 void FLevelEditorActionCallbacks::OpenContentBrowser()
diff --git a/Engine/Source/Editor/LevelEditor/Private/SLevelEditor.cpp b/Engine/Source/Editor/LevelEditor/Private/SLevelEditor.cpp
--- a/Engine/Source/Editor/LevelEditor/Private/SLevelEditor.cpp
+++ b/Engine/Source/Editor/LevelEditor/Private/SLevelEditor.cpp
@@ -555,39 +555,40 @@ void SLevelEditor::AttachSequencer( TSharedPtr<SWidget> SequencerWidget, TShared
 
 	if( !bIsReentrant )
 	{
-		TSharedRef<SDockTab> Tab = SNew(SDockTab);
-		Tab = InvokeTab(LevelEditorTabIds::Sequencer);
-
-		// Close the sequence editor after invoking a sequencer tab instead of before so that the existing asset editor doesn't refer to a stale sequencer.
-		if(SequencerAssetEditor.IsValid())
+		TSharedPtr<SDockTab> Tab = TryInvokeTab(LevelEditorTabIds::Sequencer);
+		if(Tab.IsValid())
 		{
-			// Closing the window will invoke this method again but we are handling reopening with a new movie scene ourselves
-			TGuardValue<bool> ReentrantGuard(bIsReentrant, true);
-			// Shutdown cleanly
-			SequencerAssetEditor.Pin()->CloseWindow();
-		}
+			// Close the sequence editor after invoking a sequencer tab instead of before so that the existing asset editor doesn't refer to a stale sequencer.
+			if (SequencerAssetEditor.IsValid())
+			{
+				// Closing the window will invoke this method again but we are handling reopening with a new movie scene ourselves
+				TGuardValue<bool> ReentrantGuard(bIsReentrant, true);
+				// Shutdown cleanly
+				SequencerAssetEditor.Pin()->CloseWindow();
+			}
 
-		if(!FGlobalTabmanager::Get()->OnOverrideDockableAreaRestore_Handler.IsBound())
-		{
-			// Don't allow standard tab closing behavior when the override is active
-			Tab->SetOnTabClosed(SDockTab::FOnTabClosedCallback::CreateStatic(&Local::OnSequencerClosed, TWeakPtr<IAssetEditorInstance>(NewSequencerAssetEditor)));
-		}
-		if(SequencerWidget.IsValid() && NewSequencerAssetEditor.IsValid())
-		{
-			Tab->SetContent(SequencerWidget.ToSharedRef());
-			SequencerWidgetPtr = SequencerWidget;
-			SequencerAssetEditor = NewSequencerAssetEditor;
-			if (FGlobalTabmanager::Get()->OnOverrideDockableAreaRestore_Handler.IsBound())
+			if (!FGlobalTabmanager::Get()->OnOverrideDockableAreaRestore_Handler.IsBound())
 			{
-				// @todo vreditor: more general vr editor tab manager should handle windows instead
-				// Close the original tab so we just work with the override window
-				Tab->RequestCloseTab();
+				// Don't allow standard tab closing behavior when the override is active
+				Tab->SetOnTabClosed(SDockTab::FOnTabClosedCallback::CreateStatic(&Local::OnSequencerClosed, TWeakPtr<IAssetEditorInstance>(NewSequencerAssetEditor)));
+			}
+			if (SequencerWidget.IsValid() && NewSequencerAssetEditor.IsValid())
+			{
+				Tab->SetContent(SequencerWidget.ToSharedRef());
+				SequencerWidgetPtr = SequencerWidget;
+				SequencerAssetEditor = NewSequencerAssetEditor;
+				if (FGlobalTabmanager::Get()->OnOverrideDockableAreaRestore_Handler.IsBound())
+				{
+					// @todo vreditor: more general vr editor tab manager should handle windows instead
+					// Close the original tab so we just work with the override window
+					Tab->RequestCloseTab();
+				}
+			}
+			else
+			{
+				Tab->SetContent(SNullWidget::NullWidget);
+				SequencerAssetEditor.Reset();
 			}
-		}
-		else
-		{
-			Tab->SetContent(SNullWidget::NullWidget);
-			SequencerAssetEditor.Reset();
 		}
 	}
 }
@@ -896,10 +897,10 @@ bool SLevelEditor::CanSpawnEditorModeToolboxTab(const FSpawnTabArgs& Args) const
 	return GLevelEditorModeTools().ShouldShowModeToolbox();
 }
 
-TSharedRef<SDockTab> SLevelEditor::InvokeTab( FName TabID )
+TSharedPtr<SDockTab> SLevelEditor::TryInvokeTab( FName TabID )
 {
 	TSharedPtr<FTabManager> LevelEditorTabManager = GetTabManager();
-	return LevelEditorTabManager->InvokeTab(TabID);
+	return LevelEditorTabManager->TryInvokeTab(TabID);
 }
 
 void SLevelEditor::SyncDetailsToSelection()
@@ -932,7 +933,7 @@ void SLevelEditor::SyncDetailsToSelection()
 
 		if(DetailsView->IsUpdatable() && !DetailsView->IsLocked())
 		{
-			InvokeTab(DetailsTabIdentifier);
+			TryInvokeTab(DetailsTabIdentifier);
 			return;
 		}
 	}
@@ -940,7 +941,7 @@ void SLevelEditor::SyncDetailsToSelection()
 	// If we got this far then there were no open details views, so open the first available one
 	if(!FirstClosedDetailsTabIdentifier.IsNone())
 	{
-		InvokeTab(FirstClosedDetailsTabIdentifier);
+		TryInvokeTab(FirstClosedDetailsTabIdentifier);
 	}
 }
 
diff --git a/Engine/Source/Editor/LevelEditor/Private/SLevelEditor.h b/Engine/Source/Editor/LevelEditor/Private/SLevelEditor.h
--- a/Engine/Source/Editor/LevelEditor/Private/SLevelEditor.h
+++ b/Engine/Source/Editor/LevelEditor/Private/SLevelEditor.h
@@ -77,7 +77,7 @@ public:
 	 * Given a tab ID, summons a new tab in the position saved in the current layout, or in a default position.
 	 * @return the invoked tab
 	 */
-	TSharedRef<SDockTab> InvokeTab( FName TabID );
+	TSharedPtr<SDockTab> TryInvokeTab( FName TabID );
 
 	/**
 	 * Sync the details panel to the current selection
diff --git a/Engine/Source/Editor/LocalizationDashboard/Private/LocalizationDashboard.cpp b/Engine/Source/Editor/LocalizationDashboard/Private/LocalizationDashboard.cpp
--- a/Engine/Source/Editor/LocalizationDashboard/Private/LocalizationDashboard.cpp
+++ b/Engine/Source/Editor/LocalizationDashboard/Private/LocalizationDashboard.cpp
@@ -248,7 +248,7 @@ void FLocalizationDashboard::Terminate()
 
 void FLocalizationDashboard::Show()
 {
-	FGlobalTabmanager::Get()->InvokeTab(TabName);
+	FGlobalTabmanager::Get()->TryInvokeTab(TabName);
 }
 
 TWeakPtr<SDockTab> FLocalizationDashboard::ShowTargetEditorTab(ULocalizationTarget* const LocalizationTarget)
diff --git a/Engine/Source/Editor/MainFrame/Private/Frame/MainFrameActions.cpp b/Engine/Source/Editor/MainFrame/Private/Frame/MainFrameActions.cpp
--- a/Engine/Source/Editor/MainFrame/Private/Frame/MainFrameActions.cpp
+++ b/Engine/Source/Editor/MainFrame/Private/Frame/MainFrameActions.cpp
@@ -918,7 +918,7 @@ void FMainFrameActionCallbacks::ToggleFullscreen_Execute()
 #if !PLATFORM_MAC && !PLATFORM_LINUX // Fullscreen mode in the editor is currently unsupported on Mac or Linux
 	if ( GIsEditor && FApp::HasProjectName() )
 	{
-		static TWeakPtr<SDockTab> LevelEditorTabPtr = FGlobalTabmanager::Get()->InvokeTab(FTabId("LevelEditor"));
+		static TWeakPtr<SDockTab> LevelEditorTabPtr = FGlobalTabmanager::Get()->TryInvokeTab(FTabId("LevelEditor"));
 		const TSharedPtr<SWindow> LevelEditorWindow = FSlateApplication::Get().FindWidgetWindow( LevelEditorTabPtr.Pin().ToSharedRef() );
 
 		if (LevelEditorWindow->GetWindowMode() == EWindowMode::Windowed)
@@ -1004,7 +1004,7 @@ void FMainFrameActionCallbacks::OpenSlateApp_ViaModule( FName AppName, FName Mod
 
 void FMainFrameActionCallbacks::OpenSlateApp( FName AppName )
 {
-	FGlobalTabmanager::Get()->InvokeTab(FTabId(AppName));
+	FGlobalTabmanager::Get()->TryInvokeTab(FTabId(AppName));
 }
 
 bool FMainFrameActionCallbacks::OpenSlateApp_IsChecked( FName AppName )
@@ -1144,7 +1144,7 @@ void FMainFrameActionCallbacks::CreditsUnrealEd_Execute()
 
 void FMainFrameActionCallbacks::OpenWidgetReflector_Execute()
 {
-	FGlobalTabmanager::Get()->InvokeTab(FTabId("WidgetReflector"));
+	FGlobalTabmanager::Get()->TryInvokeTab(FTabId("WidgetReflector"));
 }
 
 
diff --git a/Engine/Source/Editor/MainFrame/Private/MainFrameModule.cpp b/Engine/Source/Editor/MainFrame/Private/MainFrameModule.cpp
--- a/Engine/Source/Editor/MainFrame/Private/MainFrameModule.cpp
+++ b/Engine/Source/Editor/MainFrame/Private/MainFrameModule.cpp
@@ -228,7 +228,7 @@ void FMainFrameModule::CreateDefaultMainFrame( const bool bStartImmersive, const
 		TSharedPtr<SDockTab> MainTab;
 		if ( bLevelEditorIsMainTab )
 		{
-			MainTab = FGlobalTabmanager::Get()->InvokeTab( FTabId("LevelEditor") );
+			MainTab = FGlobalTabmanager::Get()->TryInvokeTab( FTabId("LevelEditor") );
 
 			// make sure we only allow the message log to be shown when we have a level editor main tab
 			FMessageLogModule& MessageLogModule = FModuleManager::LoadModuleChecked<FMessageLogModule>(TEXT("MessageLog"));
diff --git a/Engine/Source/Editor/MainFrame/Private/Menus/MainMenu.h b/Engine/Source/Editor/MainFrame/Private/Menus/MainMenu.h
--- a/Engine/Source/Editor/MainFrame/Private/Menus/MainMenu.h
+++ b/Engine/Source/Editor/MainFrame/Private/Menus/MainMenu.h
@@ -84,7 +84,7 @@ private:
 	*/
 	static void OpenProjectLauncher()
 	{
-		FGlobalTabmanager::Get()->InvokeTab(FName(TEXT("ProjectLauncher")));
+		FGlobalTabmanager::Get()->TryInvokeTab(FName(TEXT("ProjectLauncher")));
 	}
 
 	/**
@@ -103,6 +103,6 @@ private:
 	static void OpenVisualLogger()
 	{
 		FModuleManager::Get().LoadModuleChecked<IModuleInterface>("LogVisualizer");
-		FGlobalTabmanager::Get()->InvokeTab(FName(TEXT("VisualLogger")));
+		FGlobalTabmanager::Get()->TryInvokeTab(FName(TEXT("VisualLogger")));
 	}
 };
diff --git a/Engine/Source/Editor/MaterialEditor/Private/MaterialEditor.cpp b/Engine/Source/Editor/MaterialEditor/Private/MaterialEditor.cpp
--- a/Engine/Source/Editor/MaterialEditor/Private/MaterialEditor.cpp
+++ b/Engine/Source/Editor/MaterialEditor/Private/MaterialEditor.cpp
@@ -2222,7 +2222,7 @@ void FMaterialEditor::UpdateMaterialinfoList_Old()
 
 		if (bForceDisplay)
 		{
-			TabManager->InvokeTab(MaterialStatsManager->GetGridOldStatsTabName());
+			TabManager->TryInvokeTab(MaterialStatsManager->GetGridOldStatsTabName());
 		}
 	}
 }
@@ -3295,7 +3295,7 @@ void FMaterialEditor::OnCreateComponentMaskNode()
 
 void FMaterialEditor::OnFindInMaterial()
 {
-	TabManager->InvokeTab(FindTabId);
+	TabManager->TryInvokeTab(FindTabId);
 	FindResults->FocusForUse();
 }
 
diff --git a/Engine/Source/Editor/MaterialEditor/Private/MaterialStats.cpp b/Engine/Source/Editor/MaterialEditor/Private/MaterialStats.cpp
--- a/Engine/Source/Editor/MaterialEditor/Private/MaterialStats.cpp
+++ b/Engine/Source/Editor/MaterialEditor/Private/MaterialStats.cpp
@@ -379,7 +379,7 @@ void FMaterialStats::DisplayOldStats(const bool bShow)
 {
 	if (bShow)
 	{
-		MaterialEditor->GetTabManager()->InvokeTab(OldStatsTabId);
+		MaterialEditor->GetTabManager()->TryInvokeTab(OldStatsTabId);
 	}
 	else if (!bShowOldStats && OldStatsTab.IsValid())
 	{
@@ -391,7 +391,7 @@ void FMaterialStats::DisplayStatsGrid(const bool bShow)
 {
 	if (bShow)
 	{
-		MaterialEditor->GetTabManager()->InvokeTab(StatsTabId);
+		MaterialEditor->GetTabManager()->TryInvokeTab(StatsTabId);
 	}
 	else if (!bShowStats && StatsTab.IsValid())
 	{
diff --git a/Engine/Source/Editor/Matinee/Private/Matinee.cpp b/Engine/Source/Editor/Matinee/Private/Matinee.cpp
--- a/Engine/Source/Editor/Matinee/Private/Matinee.cpp
+++ b/Engine/Source/Editor/Matinee/Private/Matinee.cpp
@@ -259,7 +259,7 @@ void FMatinee::SetCurveTabVisibility(bool Visible)
 	}
 	else if ( !CurveEdTab.IsValid() && Visible )
 	{
-		TabManager->InvokeTab(MatineeCurveEdName);
+		TabManager->TryInvokeTab(MatineeCurveEdName);
 	}
 }
 
diff --git a/Engine/Source/Editor/Matinee/Private/MatineeMenus.cpp b/Engine/Source/Editor/Matinee/Private/MatineeMenus.cpp
--- a/Engine/Source/Editor/Matinee/Private/MatineeMenus.cpp
+++ b/Engine/Source/Editor/Matinee/Private/MatineeMenus.cpp
@@ -1194,7 +1194,7 @@ void FMatinee::OnCreateCameraActorAtCurrentCameraLocation()
 /** Called when the "Launch Custom Preview Viewport" is pressed */
 void FMatinee::OnLaunchRecordingViewport()
 {
-	FGlobalTabmanager::Get()->InvokeTab( FName("RecordingViewport") );
+	FGlobalTabmanager::Get()->TryInvokeTab( FName("RecordingViewport") );
 }
 
 TSharedRef<SDockTab> FMatinee::SpawnRecordingViewport( const FSpawnTabArgs& Args )
diff --git a/Engine/Source/Editor/ProjectSettingsViewer/Private/ProjectSettingsViewerModule.cpp b/Engine/Source/Editor/ProjectSettingsViewer/Private/ProjectSettingsViewerModule.cpp
--- a/Engine/Source/Editor/ProjectSettingsViewer/Private/ProjectSettingsViewerModule.cpp
+++ b/Engine/Source/Editor/ProjectSettingsViewer/Private/ProjectSettingsViewerModule.cpp
@@ -62,7 +62,7 @@ public:
 
 	virtual void ShowSettings( const FName& CategoryName, const FName& SectionName ) override
 	{
-		FGlobalTabmanager::Get()->InvokeTab(ProjectSettingsTabName);
+		FGlobalTabmanager::Get()->TryInvokeTab(ProjectSettingsTabName);
 		ISettingsEditorModelPtr SettingsEditorModel = SettingsEditorModelPtr.Pin();
 
 		if (SettingsEditorModel.IsValid())
diff --git a/Engine/Source/Editor/PropertyEditor/Private/Presentation/PropertyEditor/PropertyEditor.cpp b/Engine/Source/Editor/PropertyEditor/Private/Presentation/PropertyEditor/PropertyEditor.cpp
--- a/Engine/Source/Editor/PropertyEditor/Private/Presentation/PropertyEditor/PropertyEditor.cpp
+++ b/Engine/Source/Editor/PropertyEditor/Private/Presentation/PropertyEditor/PropertyEditor.cpp
@@ -309,7 +309,7 @@ void FPropertyEditor::EditConfigHierarchy()
 
 	IConfigEditorModule& ConfigEditorModule = FModuleManager::LoadModuleChecked<IConfigEditorModule>("ConfigEditor");
 	ConfigEditorModule.CreateHierarchyEditor(NodeProperty);
-	FGlobalTabmanager::Get()->InvokeTab(FName(TEXT("ConfigEditor")));
+	FGlobalTabmanager::Get()->TryInvokeTab(FName(TEXT("ConfigEditor")));
 }
 
 void FPropertyEditor::InsertItem()
diff --git a/Engine/Source/Editor/SceneOutliner/Private/WorldTreeItem.cpp b/Engine/Source/Editor/SceneOutliner/Private/WorldTreeItem.cpp
--- a/Engine/Source/Editor/SceneOutliner/Private/WorldTreeItem.cpp
+++ b/Engine/Source/Editor/SceneOutliner/Private/WorldTreeItem.cpp
@@ -125,7 +125,7 @@ void FWorldTreeItem::OnDrop(FDragDropPayload& DraggedObjects, UWorld& InWorld, c
 void FWorldTreeItem::OpenWorldSettings() const
 {
 	FLevelEditorModule& LevelEditorModule = FModuleManager::GetModuleChecked<FLevelEditorModule>( TEXT("LevelEditor") );
-	LevelEditorModule.GetLevelEditorTabManager()->InvokeTab(FName("WorldSettingsTab"));	
+	LevelEditorModule.GetLevelEditorTabManager()->TryInvokeTab(FName("WorldSettingsTab"));	
 }
 
 }		// namespace SceneOutliner
diff --git a/Engine/Source/Editor/Sequencer/Private/LevelEditorSequencerIntegration.cpp b/Engine/Source/Editor/Sequencer/Private/LevelEditorSequencerIntegration.cpp
--- a/Engine/Source/Editor/Sequencer/Private/LevelEditorSequencerIntegration.cpp
+++ b/Engine/Source/Editor/Sequencer/Private/LevelEditorSequencerIntegration.cpp
@@ -1018,7 +1018,7 @@ void FLevelEditorSequencerIntegration::DetachOutlinerColumn()
 		{
 			if (LevelEditorTabManager->GetOwnerTab().IsValid())
 			{
-				LevelEditorTabManager->InvokeTab(FName("LevelEditorSceneOutliner"))->RequestCloseTab();			
+				LevelEditorTabManager->TryInvokeTab(FName("LevelEditorSceneOutliner"))->RequestCloseTab();			
 			}
 		}
 		
@@ -1026,7 +1026,7 @@ void FLevelEditorSequencerIntegration::DetachOutlinerColumn()
 		{
 			if (LevelEditorTabManager->GetOwnerTab().IsValid())
 			{
-				LevelEditorTabManager->InvokeTab(FName("LevelEditorSceneOutliner"));
+				LevelEditorTabManager->TryInvokeTab(FName("LevelEditorSceneOutliner"));
 			}
 		}
 	}
diff --git a/Engine/Source/Editor/Sequencer/Private/SSequencer.cpp b/Engine/Source/Editor/Sequencer/Private/SSequencer.cpp
--- a/Engine/Source/Editor/Sequencer/Private/SSequencer.cpp
+++ b/Engine/Source/Editor/Sequencer/Private/SSequencer.cpp
@@ -3190,15 +3190,18 @@ void SSequencer::OnCurveEditorVisibilityChanged(bool bShouldBeVisible)
 	{
 		// Request the Tab Manager invoke the tab. This will spawn the tab if needed, otherwise pull it to focus. This assumes
 		// that the Toolkit Host's Tab Manager has already registered a tab with a NullWidget for content.
-		TSharedRef<SDockTab> CurveEditorTab = Sequencer->GetToolkitHost()->GetTabManager()->InvokeTab(TabId);
-		CurveEditorTab->SetContent(CurveEditorPanel.ToSharedRef());
+		TSharedPtr<SDockTab> CurveEditorTab = Sequencer->GetToolkitHost()->GetTabManager()->TryInvokeTab(TabId);
+		if (CurveEditorTab.IsValid())
+		{
+			CurveEditorTab->SetContent(CurveEditorPanel.ToSharedRef());
 
-		const FSlateIcon SequencerGraphIcon = FSlateIcon(FEditorStyle::GetStyleSetName(), "GenericCurveEditor.TabIcon");
-		CurveEditorTab->SetTabIcon(SequencerGraphIcon.GetIcon());
+			const FSlateIcon SequencerGraphIcon = FSlateIcon(FEditorStyle::GetStyleSetName(), "GenericCurveEditor.TabIcon");
+			CurveEditorTab->SetTabIcon(SequencerGraphIcon.GetIcon());
 
-		CurveEditorTab->SetLabel(LOCTEXT("SequencerMainGraphEditorTitle", "Sequencer Curves"));
+			CurveEditorTab->SetLabel(LOCTEXT("SequencerMainGraphEditorTitle", "Sequencer Curves"));
 
-		SequencerPtr.Pin()->GetCurveEditor()->ZoomToFit();
+			SequencerPtr.Pin()->GetCurveEditor()->ZoomToFit();
+		}
 	}
 	else
 	{
diff --git a/Engine/Source/Editor/StaticMeshEditor/Private/StaticMeshEditor.cpp b/Engine/Source/Editor/StaticMeshEditor/Private/StaticMeshEditor.cpp
--- a/Engine/Source/Editor/StaticMeshEditor/Private/StaticMeshEditor.cpp
+++ b/Engine/Source/Editor/StaticMeshEditor/Private/StaticMeshEditor.cpp
@@ -342,7 +342,7 @@ void FStaticMeshEditor::GenerateSecondaryToolbar()
 	if (!Tab)
 	{
 		// By default, the tab is closed but we want it to be opened when it is populated
-		Tab = TSharedPtr<SDockTab>(TabManager->InvokeTab(SecondaryToolbarTabId));
+		Tab = TSharedPtr<SDockTab>(TabManager->TryInvokeTab(SecondaryToolbarTabId));
 	}
 
 	// Override the display name if it was set
@@ -352,7 +352,7 @@ void FStaticMeshEditor::GenerateSecondaryToolbar()
 	}
 
 	// But have the focus on the default toolbar
-	TabManager->InvokeTab(GetToolbarTabId());
+	TabManager->TryInvokeTab(GetToolbarTabId());
 }
 
 void FStaticMeshEditor::AddSecondaryToolbarExtender(TSharedPtr<FExtender> Extender)
@@ -2215,7 +2215,7 @@ FEditorViewportClient& FStaticMeshEditor::GetViewportClient()
 
 void FStaticMeshEditor::OnConvexDecomposition()
 {
-	TabManager->InvokeTab(CollisionTabId);
+	TabManager->TryInvokeTab(CollisionTabId);
 }
 
 bool FStaticMeshEditor::OnRequestClose()
diff --git a/Engine/Source/Editor/TranslationEditor/Private/TranslationEditor.cpp b/Engine/Source/Editor/TranslationEditor/Private/TranslationEditor.cpp
--- a/Engine/Source/Editor/TranslationEditor/Private/TranslationEditor.cpp
+++ b/Engine/Source/Editor/TranslationEditor/Private/TranslationEditor.cpp
@@ -1339,7 +1339,7 @@ void FTranslationEditor::ImportFromPoFile(FString FileToImport)
 		{
 			DataManager->LoadFromArchive(DataManager->GetAllTranslationsArray(), true, true);
 
-			TabManager->InvokeTab(ChangedOnImportTabId);
+			TabManager->TryInvokeTab(ChangedOnImportTabId);
 			ChangedOnImportPropertyTable->SetObjects((TArray<UObject*>&)DataManager->GetChangedOnImportArray());
 			// Need to re-add the columns we want to display
 			ChangedOnImportPropertyTable->AddColumn((TWeakFieldPtr<FProperty>)FindField<FProperty>(UTranslationUnit::StaticClass(), "Source"));
@@ -1397,7 +1397,7 @@ void FTranslationEditor::OnFilterTextCommitted(const FText& InFilterText, ETextC
 
 void FTranslationEditor::OpenSearchTab_Execute()
 {
-	TabManager->InvokeTab(SearchTabId);
+	TabManager->TryInvokeTab(SearchTabId);
 }
 
 FReply FTranslationEditor::OnGetHistoryButtonClicked()
diff --git a/Engine/Source/Editor/UATHelper/UATHelperModule.cpp b/Engine/Source/Editor/UATHelper/UATHelperModule.cpp
--- a/Engine/Source/Editor/UATHelper/UATHelperModule.cpp
+++ b/Engine/Source/Editor/UATHelper/UATHelperModule.cpp
@@ -524,7 +524,7 @@ public:
 
 	static void HandleUatHyperlinkNavigate()
 	{
-		FGlobalTabmanager::Get()->InvokeTab(FName("OutputLog"));
+		FGlobalTabmanager::Get()->TryInvokeTab(FName("OutputLog"));
 	}
 
 	static void HandleUatCancelButtonClicked(TSharedPtr<FMonitoredProcess> PackagerProcess)
diff --git a/Engine/Source/Editor/UMGEditor/Private/WidgetBlueprintEditorToolbar.cpp b/Engine/Source/Editor/UMGEditor/Private/WidgetBlueprintEditorToolbar.cpp
--- a/Engine/Source/Editor/UMGEditor/Private/WidgetBlueprintEditorToolbar.cpp
+++ b/Engine/Source/Editor/UMGEditor/Private/WidgetBlueprintEditorToolbar.cpp
@@ -134,7 +134,7 @@ void FWidgetBlueprintEditorToolbar::AddWidgetReflector(UToolMenu* InMenu)
 	Section.AddEntry(FToolMenuEntry::InitToolBarButton(
 		"OpenWidgetReflector",
 		FUIAction(
-			FExecuteAction::CreateLambda([=] { FGlobalTabmanager::Get()->InvokeTab(FTabId("WidgetReflector")); }),
+			FExecuteAction::CreateLambda([=] { FGlobalTabmanager::Get()->TryInvokeTab(FTabId("WidgetReflector")); }),
 			FCanExecuteAction()
 		)
 		, LOCTEXT("OpenWidgetReflector", "Widget Reflector")
diff --git a/Engine/Source/Editor/UndoHistory/Public/UndoHistoryModule.h b/Engine/Source/Editor/UndoHistory/Public/UndoHistoryModule.h
--- a/Engine/Source/Editor/UndoHistory/Public/UndoHistoryModule.h
+++ b/Engine/Source/Editor/UndoHistory/Public/UndoHistoryModule.h
@@ -25,7 +25,7 @@ public:
 
 	static void ExecuteOpenUndoHistory()
 	{
-		FGlobalTabmanager::Get()->InvokeTab(UndoHistoryTabName);
+		FGlobalTabmanager::Get()->TryInvokeTab(UndoHistoryTabName);
 	}
 
 private:
diff --git a/Engine/Source/Editor/UnrealEd/Private/EditorModeManager.cpp b/Engine/Source/Editor/UnrealEd/Private/EditorModeManager.cpp
--- a/Engine/Source/Editor/UnrealEd/Private/EditorModeManager.cpp
+++ b/Engine/Source/Editor/UnrealEd/Private/EditorModeManager.cpp
@@ -819,7 +819,7 @@ void FEditorModeTools::SpawnOrUpdateModeToolbar()
 		}
 		else if (ToolkitHost.IsValid())
 		{
-			ToolkitHost.Pin()->GetTabManager()->InvokeTab(EditorModeToolbarTabName);
+			ToolkitHost.Pin()->GetTabManager()->TryInvokeTab(EditorModeToolbarTabName);
 		}
 	}
 }
diff --git a/Engine/Source/Editor/UnrealEd/Private/Fbx/SFbxSceneOptionWindow.cpp b/Engine/Source/Editor/UnrealEd/Private/Fbx/SFbxSceneOptionWindow.cpp
--- a/Engine/Source/Editor/UnrealEd/Private/Fbx/SFbxSceneOptionWindow.cpp
+++ b/Engine/Source/Editor/UnrealEd/Private/Fbx/SFbxSceneOptionWindow.cpp
@@ -1319,17 +1319,17 @@ void SFbxSceneOptionWindow::Construct(const FArguments& InArgs)
 	//By default we want to see the Scene tab
 	if (!SceneInfoOriginal.IsValid())
 	{
-		FbxSceneImportTabManager->InvokeTab(FTabId("Scene"));
+		FbxSceneImportTabManager->TryInvokeTab(FTabId("Scene"));
 	}
 	else
 	{
 		if (bCanReimportHierarchy)
 		{
-			FbxSceneImportTabManager->InvokeTab(FTabId("SceneReImport"));
+			FbxSceneImportTabManager->TryInvokeTab(FTabId("SceneReImport"));
 		}
 		else
 		{
-			FbxSceneImportTabManager->InvokeTab(FTabId("StaticMeshesReimport"));
+			FbxSceneImportTabManager->TryInvokeTab(FTabId("StaticMeshesReimport"));
 		}
 	}
 }
diff --git a/Engine/Source/Editor/UnrealEd/Private/Kismet2/DebuggerCommands.cpp b/Engine/Source/Editor/UnrealEd/Private/Kismet2/DebuggerCommands.cpp
--- a/Engine/Source/Editor/UnrealEd/Private/Kismet2/DebuggerCommands.cpp
+++ b/Engine/Source/Editor/UnrealEd/Private/Kismet2/DebuggerCommands.cpp
@@ -2035,12 +2035,12 @@ void FInternalPlayWorldCommandCallbacks::PlayInSettings_Clicked()
 
 void FInternalPlayWorldCommandCallbacks::OpenProjectLauncher_Clicked()
 {
-	FGlobalTabmanager::Get()->InvokeTab(FTabId("ProjectLauncher"));
+	FGlobalTabmanager::Get()->TryInvokeTab(FTabId("ProjectLauncher"));
 }
 
 void FInternalPlayWorldCommandCallbacks::OpenDeviceManager_Clicked()
 {
-	FGlobalTabmanager::Get()->InvokeTab(FTabId("DeviceManager"));
+	FGlobalTabmanager::Get()->TryInvokeTab(FTabId("DeviceManager"));
 }
 
 void FInternalPlayWorldCommandCallbacks::RepeatLastLaunch_Clicked()
diff --git a/Engine/Source/Editor/UnrealEd/Private/Toolkits/AssetEditorToolkit.cpp b/Engine/Source/Editor/UnrealEd/Private/Toolkits/AssetEditorToolkit.cpp
--- a/Engine/Source/Editor/UnrealEd/Private/Toolkits/AssetEditorToolkit.cpp
+++ b/Engine/Source/Editor/UnrealEd/Private/Toolkits/AssetEditorToolkit.cpp
@@ -449,7 +449,7 @@ bool FAssetEditorToolkit::CloseWindow()
 
 void FAssetEditorToolkit::InvokeTab(const FTabId& TabId)
 {
-	GetTabManager()->InvokeTab(TabId);
+	GetTabManager()->TryInvokeTab(TabId);
 }
 
 TSharedPtr<class FTabManager> FAssetEditorToolkit::GetAssociatedTabManager()
diff --git a/Engine/Source/Editor/UnrealEd/Private/UnrealEdEngine.cpp b/Engine/Source/Editor/UnrealEd/Private/UnrealEdEngine.cpp
--- a/Engine/Source/Editor/UnrealEd/Private/UnrealEdEngine.cpp
+++ b/Engine/Source/Editor/UnrealEd/Private/UnrealEdEngine.cpp
@@ -1058,7 +1058,7 @@ void UUnrealEdEngine::ShowLightingStaticMeshInfoWindow()
 	// first invoke the stats viewer tab
 	FLevelEditorModule& LevelEditorModule = FModuleManager::GetModuleChecked<FLevelEditorModule>( TEXT("LevelEditor") );
 	TSharedPtr<FTabManager> LevelEditorTabManager = LevelEditorModule.GetLevelEditorTabManager();
-	LevelEditorTabManager->InvokeTab(FName("LevelEditorStatsViewer"));
+	LevelEditorTabManager->TryInvokeTab(FName("LevelEditorStatsViewer"));
 
 	// then switch pages
 	FStatsViewerModule& StatsViewerModule = FModuleManager::Get().LoadModuleChecked<FStatsViewerModule>(TEXT("StatsViewer"));
@@ -1071,7 +1071,7 @@ void UUnrealEdEngine::OpenSceneStatsWindow()
 	// first invoke the stats viewer tab
 	FLevelEditorModule& LevelEditorModule = FModuleManager::GetModuleChecked<FLevelEditorModule>( TEXT("LevelEditor") );
 	TSharedPtr<FTabManager> LevelEditorTabManager = LevelEditorModule.GetLevelEditorTabManager();
-	LevelEditorTabManager->InvokeTab(FName("LevelEditorStatsViewer"));
+	LevelEditorTabManager->TryInvokeTab(FName("LevelEditorStatsViewer"));
 
 	// then switch pages
 	FStatsViewerModule& StatsViewerModule = FModuleManager::Get().LoadModuleChecked<FStatsViewerModule>(TEXT("StatsViewer"));
@@ -1084,7 +1084,7 @@ void UUnrealEdEngine::OpenTextureStatsWindow()
 	// first invoke the stats viewer tab
 	FLevelEditorModule& LevelEditorModule = FModuleManager::GetModuleChecked<FLevelEditorModule>( TEXT("LevelEditor") );
 	TSharedPtr<FTabManager> LevelEditorTabManager = LevelEditorModule.GetLevelEditorTabManager();
-	LevelEditorTabManager->InvokeTab(FName("LevelEditorStatsViewer"));
+	LevelEditorTabManager->TryInvokeTab(FName("LevelEditorStatsViewer"));
 
 	// then switch pages
 	FStatsViewerModule& StatsViewerModule = FModuleManager::Get().LoadModuleChecked<FStatsViewerModule>(TEXT("StatsViewer"));
diff --git a/Engine/Source/Editor/VREditor/VREditorMode.cpp b/Engine/Source/Editor/VREditor/VREditorMode.cpp
--- a/Engine/Source/Editor/VREditor/VREditorMode.cpp
+++ b/Engine/Source/Editor/VREditor/VREditorMode.cpp
@@ -287,7 +287,7 @@ void UVREditorMode::Enter()
 	// Switch us back to placement mode and close any open sequencer windows
 	FVREditorActionCallbacks::ChangeEditorModes(FBuiltinEditorModes::EM_Placement);
 	FLevelEditorModule& LevelEditorModule = FModuleManager::GetModuleChecked<FLevelEditorModule>(TEXT("LevelEditor"));
-	LevelEditorModule.GetLevelEditorTabManager()->InvokeTab(FTabId("Sequencer"))->RequestCloseTab();
+	LevelEditorModule.GetLevelEditorTabManager()->TryInvokeTab(FTabId("Sequencer"))->RequestCloseTab();
 
 	// Setup our avatar
 	if (AvatarActor == nullptr)
diff --git a/Engine/Source/Runtime/Engine/Private/Collision/WorldCollision.cpp b/Engine/Source/Runtime/Engine/Private/Collision/WorldCollision.cpp
--- a/Engine/Source/Runtime/Engine/Private/Collision/WorldCollision.cpp
+++ b/Engine/Source/Runtime/Engine/Private/Collision/WorldCollision.cpp
@@ -461,7 +461,7 @@ public:
 #if ENABLE_COLLISION_ANALYZER
 		if (FParse::Command(&Cmd, TEXT("CANALYZER")))
 		{
-			FGlobalTabmanager::Get()->InvokeTab(FName(TEXT("CollisionAnalyzerApp")));
+			FGlobalTabmanager::Get()->TryInvokeTab(FName(TEXT("CollisionAnalyzerApp")));
 			return true;
 		}
 #endif // ENABLE_COLLISION_ANALYZER
diff --git a/Engine/Source/Runtime/Engine/Private/VisualLogger/VisualLogger.cpp b/Engine/Source/Runtime/Engine/Private/VisualLogger/VisualLogger.cpp
--- a/Engine/Source/Runtime/Engine/Private/VisualLogger/VisualLogger.cpp
+++ b/Engine/Source/Runtime/Engine/Private/VisualLogger/VisualLogger.cpp
@@ -650,7 +650,7 @@ public:
 #if WITH_EDITOR
 				else
 				{
-					FGlobalTabmanager::Get()->InvokeTab(FName(TEXT("VisualLogger")));
+					FGlobalTabmanager::Get()->TryInvokeTab(FName(TEXT("VisualLogger")));
 					return true;
 				}
 #endif
diff --git a/Engine/Source/Runtime/Slate/Private/Framework/Docking/TabManager.cpp b/Engine/Source/Runtime/Slate/Private/Framework/Docking/TabManager.cpp
--- a/Engine/Source/Runtime/Slate/Private/Framework/Docking/TabManager.cpp
+++ b/Engine/Source/Runtime/Slate/Private/Framework/Docking/TabManager.cpp
@@ -1022,7 +1022,7 @@ void FTabManager::PopulateTabSpawnerMenu( FMenuBuilder& PopulateMe, TSharedRef<F
 		const TSharedRef<FTabSpawnerEntry>& SpawnerEntry = SpawnerIterator.Value();
 		if ( SpawnerEntry->bAutoGenerateMenuEntry )
 		{
-			if (SpawnerEntry->TabType == NAME_None || TabBlacklist->PassesFilter(SpawnerEntry->TabType))
+			if (IsAllowedTab(SpawnerEntry->TabType))
 			{
 				AllSpawners->AddUnique(SpawnerEntry);
 			}
@@ -1035,7 +1035,7 @@ void FTabManager::PopulateTabSpawnerMenu( FMenuBuilder& PopulateMe, TSharedRef<F
 		const TSharedRef<FTabSpawnerEntry>& SpawnerEntry = SpawnerIterator.Value();
 		if ( SpawnerEntry->bAutoGenerateMenuEntry )
 		{
-			if (SpawnerEntry->TabType == NAME_None || TabBlacklist->PassesFilter(SpawnerEntry->TabType))
+			if (IsAllowedTab(SpawnerEntry->TabType))
 			{
 				AllSpawners->AddUnique(SpawnerEntry);
 			}
@@ -1147,20 +1147,33 @@ void FTabManager::RestoreDocumentTab( FName PlaceholderId, ESearchPreference::Ty
 	}
 }
 
-TSharedRef<SDockTab> FTabManager::InvokeTab( const FTabId& TabId )
+TSharedRef<SDockTab> FTabManager::InvokeTab(const FTabId& TabId)
 {
-	TSharedPtr<SDockTab> NewTab = InvokeTab_Internal( TabId );
-	check(NewTab.IsValid());
+	if (TSharedPtr<SDockTab> NewTab = TryInvokeTab(TabId))
+	{
+		return NewTab.ToSharedRef();
+	}
+
+	return SNew(SDockTab);
+}
+
+TSharedPtr<SDockTab> FTabManager::TryInvokeTab(const FTabId& TabId)
+{
+	TSharedPtr<SDockTab> NewTab = InvokeTab_Internal(TabId);
+	if (!NewTab.IsValid())
+	{
+		return NewTab;
+	}
+
 	TSharedPtr<SWindow> ParentWindowPtr = NewTab->GetParentWindow();
 	if ((NewTab->GetTabRole() == ETabRole::MajorTab || NewTab->GetTabRole() == ETabRole::NomadTab) && ParentWindowPtr.IsValid() && ParentWindowPtr != FGlobalTabmanager::Get()->GetRootWindow())
 	{
-		ParentWindowPtr->SetTitle( NewTab->GetTabLabel() );
+		ParentWindowPtr->SetTitle(NewTab->GetTabLabel());
 	}
 #if PLATFORM_MAC
 	FPlatformApplicationMisc::bChachedMacMenuStateNeedsUpdate = true;
 #endif
-	// Note: we expect tabs that are manually invoked to always succeed
-	return NewTab.ToSharedRef();
+	return NewTab;
 }
 
 TSharedPtr<SDockTab> FTabManager::InvokeTab_Internal( const FTabId& TabId )
@@ -1182,6 +1195,12 @@ TSharedPtr<SDockTab> FTabManager::InvokeTab_Internal( const FTabId& TabId )
 	//         v
 	//     * Spawn in a new window.
 
+	if (!IsAllowedTab(TabId))
+	{
+		UE_LOG(LogTabManager, Warning, TEXT("Cannot spawn tab for '%s'"), *(TabId.ToString()));
+		return nullptr;
+	}
+
 	TSharedPtr<FTabSpawnerEntry> Spawner = FindTabSpawnerFor(TabId.TabType);
 
 	if ( !Spawner.IsValid() )
@@ -1280,7 +1299,7 @@ TSharedPtr<SDockingTabStack> FTabManager::FindPotentiallyClosedTab( const FTabId
 
 void FTabManager::InvokeTabForMenu( FName TabId )
 {
-	InvokeTab(TabId);
+	TryInvokeTab(TabId);
 }
 
 void FTabManager::InsertDocumentTab(FName PlaceholderId, const FSearchPreference& SearchPreference, const TSharedRef<SDockTab>& UnmanagedTab, bool bPlaySpawnAnim)
@@ -1633,7 +1652,7 @@ TSharedRef<FBlacklistNames>& FTabManager::GetTabBlacklist()
 
 bool FTabManager::IsValidTabForSpawning( const FTab& SomeTab ) const
 {
-	if (SomeTab.TabId.TabType != NAME_None && !TabBlacklist->PassesFilter(SomeTab.TabId.TabType))
+	if (!IsAllowedTab(SomeTab.TabId))
 	{
 		return false;
 	}
@@ -1643,6 +1662,16 @@ bool FTabManager::IsValidTabForSpawning( const FTab& SomeTab ) const
 	return ( !NomadSpawner || !NomadSpawner->Get().IsSoleTabInstanceSpawned() );
 }
 
+bool FTabManager::IsAllowedTab(const FTabId& TabId) const
+{
+	return IsAllowedTabType(TabId.TabType);
+}
+
+bool FTabManager::IsAllowedTabType(const FName TabType) const
+{
+	return TabType == NAME_None || TabBlacklist->PassesFilter(TabType);
+}
+
 TSharedPtr<SDockTab> FTabManager::SpawnTab(const FTabId& TabId, const TSharedPtr<SWindow>& ParentWindow, const bool bCanOutputBeNullptr)
 {
 	TSharedPtr<SDockTab> NewTabWidget;
diff --git a/Engine/Source/Runtime/Slate/Public/Framework/Docking/TabManager.h b/Engine/Source/Runtime/Slate/Public/Framework/Docking/TabManager.h
--- a/Engine/Source/Runtime/Slate/Public/Framework/Docking/TabManager.h
+++ b/Engine/Source/Runtime/Slate/Public/Framework/Docking/TabManager.h
@@ -742,7 +742,16 @@ class SLATE_API FTabManager : public TSharedFromThis<FTabManager>
 		 * @param TabId The tab identifier.
 		 * @return The existing or newly spawned tab instance.
 		 */
-		virtual TSharedRef<SDockTab> InvokeTab( const FTabId& TabId );
+		UE_DEPRECATED(4.26, "FTabManager::InvokeTab is deprecated. Please use TryInvokeTab instead!")
+		virtual TSharedRef<SDockTab> InvokeTab(const FTabId& TabId);
+
+		/**
+		 * Try to open tab if it is closed at the last known location.  If it already exists, it will draw attention to the tab.
+		 *
+		 * @param TabId The tab identifier.
+		 * @return The existing or newly spawned tab instance if successful.
+		 */
+		virtual TSharedPtr<SDockTab> TryInvokeTab(const FTabId& TabId);
 
 		/**
 		 * Finds the first instance of an existing tab with the given tab id.
@@ -861,6 +870,8 @@ class SLATE_API FTabManager : public TSharedFromThis<FTabManager>
 		void RestoreSplitterContent(const TSharedRef<FSplitter>& SplitterNode, const TSharedRef<class SDockingSplitter>& SplitterWidget, const TSharedPtr<SWindow>& ParentWindow);
 		
 		bool IsValidTabForSpawning( const FTab& SomeTab ) const;
+		bool IsAllowedTab(const FTabId& TabId) const;
+		bool IsAllowedTabType(const FName TabType) const;
 		TSharedPtr<SDockTab> SpawnTab(const FTabId& TabId, const TSharedPtr<SWindow>& ParentWindow, const bool bCanOutputBeNullptr = false);
 
 		TSharedPtr<class SDockingTabStack> FindTabInLiveAreas( const FTabMatcher& TabMatcher ) const;
diff --git a/Engine/Source/Runtime/Toolbox/Private/ToolboxModule.cpp b/Engine/Source/Runtime/Toolbox/Private/ToolboxModule.cpp
--- a/Engine/Source/Runtime/Toolbox/Private/ToolboxModule.cpp
+++ b/Engine/Source/Runtime/Toolbox/Private/ToolboxModule.cpp
@@ -214,7 +214,7 @@ class FToolboxModule : public IToolboxModule
 
 	virtual void SummonToolbox() override
 	{
-		FGlobalTabmanager::Get()->InvokeTab(FTabId("DebugTools"));
+		FGlobalTabmanager::Get()->TryInvokeTab(FTabId("DebugTools"));
 	}
 };
 
