Removing uses of master/slave found in Merge with base folder Developer

diff --git a/Engine/Source/Developer/Merge/Private/SBlueprintMerge.cpp b/Engine/Source/Developer/Merge/Private/SBlueprintMerge.cpp
--- a/Engine/Source/Developer/Merge/Private/SBlueprintMerge.cpp
+++ b/Engine/Source/Developer/Merge/Private/SBlueprintMerge.cpp
@@ -267,12 +267,12 @@ UBlueprint* SBlueprintMerge::GetTargetBlueprint()
 
 void SBlueprintMerge::NextDiff()
 {
-	DiffTreeView::HighlightNextDifference(DifferencesTreeView.ToSharedRef(), RealDifferences, MasterDifferencesList);
+	DiffTreeView::HighlightNextDifference(DifferencesTreeView.ToSharedRef(), RealDifferences, PrimaryDifferencesList);
 }
 
 void SBlueprintMerge::PrevDiff()
 {
-	 DiffTreeView::HighlightPrevDifference( DifferencesTreeView.ToSharedRef(), RealDifferences, MasterDifferencesList );
+	 DiffTreeView::HighlightPrevDifference( DifferencesTreeView.ToSharedRef(), RealDifferences, PrimaryDifferencesList );
 }
 
 bool SBlueprintMerge::HasNextDiff() const
@@ -287,12 +287,12 @@ bool SBlueprintMerge::HasPrevDiff() const
 
 void SBlueprintMerge::NextConflict()
 {
-	DiffTreeView::HighlightNextDifference(DifferencesTreeView.ToSharedRef(), MergeConflicts, MasterDifferencesList);
+	DiffTreeView::HighlightNextDifference(DifferencesTreeView.ToSharedRef(), MergeConflicts, PrimaryDifferencesList);
 }
 
 void SBlueprintMerge::PrevConflict()
 {
-	DiffTreeView::HighlightPrevDifference(DifferencesTreeView.ToSharedRef(), MergeConflicts, MasterDifferencesList);
+	DiffTreeView::HighlightPrevDifference(DifferencesTreeView.ToSharedRef(), MergeConflicts, PrimaryDifferencesList);
 }
 
 bool SBlueprintMerge::HasNextConflict() const
@@ -307,7 +307,7 @@ bool SBlueprintMerge::HasPrevConflict() const
 
 void SBlueprintMerge::OnStartMerge()
 {
-	DifferencesTreeView = DiffTreeView::CreateTreeView(&MasterDifferencesList);
+	DifferencesTreeView = DiffTreeView::CreateTreeView(&PrimaryDifferencesList);
 	TreeViewContainer->SetContent(DifferencesTreeView.ToSharedRef());
 
 	if (Data.BlueprintRemote == nullptr)
@@ -362,7 +362,7 @@ void SBlueprintMerge::OnStartMerge()
 		auto DetailsView = SNew(SMergeDetailsView
 			, Data
 			, FOnMergeNodeSelected::CreateStatic(DetailsSelected, Data.OwningEditor)
-			, MasterDifferencesList
+			, PrimaryDifferencesList
 			, RealDifferences
 			, MergeConflicts);
 		DetailsControl = DetailsView;
@@ -374,7 +374,7 @@ void SBlueprintMerge::OnStartMerge()
 		auto TreeView = SNew(SMergeTreeView
 			, Data
 			, FOnMergeNodeSelected::CreateStatic(ComponentSelected, Data.OwningEditor)
-			, MasterDifferencesList
+			, PrimaryDifferencesList
 			, RealDifferences
 			, MergeConflicts);
 		TreeControl = TreeView;
@@ -386,7 +386,7 @@ void SBlueprintMerge::OnStartMerge()
 		auto GraphView = SNew(SMergeGraphView
 							, Data
 							, FOnMergeNodeSelected::CreateStatic(GraphSelected, Data.OwningEditor)
-							, MasterDifferencesList
+							, PrimaryDifferencesList
 							, RealDifferences
 							, MergeConflicts);
 		GraphControl = GraphView;
diff --git a/Engine/Source/Developer/Merge/Private/SBlueprintMerge.h b/Engine/Source/Developer/Merge/Private/SBlueprintMerge.h
--- a/Engine/Source/Developer/Merge/Private/SBlueprintMerge.h
+++ b/Engine/Source/Developer/Merge/Private/SBlueprintMerge.h
@@ -114,7 +114,7 @@ private:
 	TSharedPtr< STreeView< TSharedPtr< class FBlueprintDifferenceTreeEntry > > > DifferencesTreeView;
 
 	/** List of differences collected across all panels: */
-	TArray< TSharedPtr<class FBlueprintDifferenceTreeEntry> > MasterDifferencesList;
+	TArray< TSharedPtr<class FBlueprintDifferenceTreeEntry> > PrimaryDifferencesList;
 
 	/** List of all differences, cached so that we can iterate only the differences and not labels, etc: */
 	TArray< TSharedPtr<class FBlueprintDifferenceTreeEntry> > RealDifferences;
